Fixing sanbox wrapper not looking for cooked data in the right directory.

diff --git a/Engine/Source/Runtime/SandboxFile/Private/IPlatformFileSandboxWrapper.cpp b/Engine/Source/Runtime/SandboxFile/Private/IPlatformFileSandboxWrapper.cpp
--- a/Engine/Source/Runtime/SandboxFile/Private/IPlatformFileSandboxWrapper.cpp
+++ b/Engine/Source/Runtime/SandboxFile/Private/IPlatformFileSandboxWrapper.cpp
@@ -15,7 +15,7 @@ FSandboxPlatformFile::FSandboxPlatformFile(bool bInEntireEngineWillUseThisSandbo
 
 static FString GetCookedSandboxDir()
 {
-	return FPaths::Combine(*(FPaths::GameDir()), TEXT("Cooked"), ANSI_TO_TCHAR(FPlatformProperties::PlatformName()));
+	return FPaths::Combine(*(FPaths::GameSavedDir()), TEXT("Cooked"), ANSI_TO_TCHAR(FPlatformProperties::PlatformName()));
 }
 
 bool FSandboxPlatformFile::ShouldBeUsed(IPlatformFile* Inner, const TCHAR* CmdLine) const
