UE4 Refactoring. Changed OVERRIDE and FINAL macros to keywords override and final.

diff --git a/Engine/Source/Runtime/OpenGLDrv/Public/OpenGLDrv.h b/Engine/Source/Runtime/OpenGLDrv/Public/OpenGLDrv.h
--- a/Engine/Source/Runtime/OpenGLDrv/Public/OpenGLDrv.h
+++ b/Engine/Source/Runtime/OpenGLDrv/Public/OpenGLDrv.h
@@ -181,14 +181,14 @@ public:
 	 * Returns the time in ms that the GPU spent in this draw event.  
 	 * This blocks the CPU if necessary, so can cause hitching.
 	 */
-	float GetTiming() OVERRIDE;
+	float GetTiming() override;
 
-	virtual void StartTiming() OVERRIDE
+	virtual void StartTiming() override
 	{
 		Timing.StartTiming();
 	}
 
-	virtual void StopTiming() OVERRIDE
+	virtual void StopTiming() override
 	{
 		Timing.EndTiming();
 	}
@@ -217,15 +217,15 @@ public:
 	}
 
 	/** Start this frame of per tracking */
-	void StartFrame() OVERRIDE;
+	void StartFrame() override;
 
 	/** End this frame of per tracking, but do not block yet */
-	void EndFrame() OVERRIDE;
+	void EndFrame() override;
 
 	/** Calculates root timing base frequency (if needed by this RHI) */
-	virtual float GetRootTimingResults() OVERRIDE;
+	virtual float GetRootTimingResults() override;
 
-	virtual void LogDisjointQuery() OVERRIDE;
+	virtual void LogDisjointQuery() override;
 
 	/** Timer tracking inclusive time spent in the root nodes. */
 	FOpenGLBufferedGPUTiming RootEventTiming;
@@ -266,7 +266,7 @@ struct FOpenGLGPUProfiler : public FGPUProfiler
 		}
 	}
 
-	virtual FGPUProfilerEventNode* CreateEventNode(const TCHAR* InName, FGPUProfilerEventNode* InParent) OVERRIDE
+	virtual FGPUProfilerEventNode* CreateEventNode(const TCHAR* InName, FGPUProfilerEventNode* InParent) override
 	{
 		FOpenGLEventNode* EventNode = new FOpenGLEventNode(InName, InParent, OpenGLRHI);
 		return EventNode;
@@ -274,8 +274,8 @@ struct FOpenGLGPUProfiler : public FGPUProfiler
 
 	void Cleanup();
 
-	virtual void PushEvent(const TCHAR* Name) OVERRIDE;
-	virtual void PopEvent() OVERRIDE;
+	virtual void PushEvent(const TCHAR* Name) override;
+	virtual void PopEvent() override;
 
 	void BeginFrame(class FOpenGLDynamicRHI* InRHI);
 	void EndFrame();
@@ -495,12 +495,12 @@ class FOpenGLDynamicRHIModule : public IDynamicRHIModule
 public:
 	
 	// IModuleInterface
-	virtual bool SupportsDynamicReloading() OVERRIDE { return false; }
+	virtual bool SupportsDynamicReloading() override { return false; }
 
 	// IDynamicRHIModule
-	virtual bool IsSupported() OVERRIDE;
+	virtual bool IsSupported() override;
 
-	virtual FDynamicRHI* CreateRHI() OVERRIDE
+	virtual FDynamicRHI* CreateRHI() override
 	{
 		return new FOpenGLDynamicRHI();
 	}
diff --git a/Engine/Source/Runtime/OpenGLDrv/Public/OpenGLResources.h b/Engine/Source/Runtime/OpenGLDrv/Public/OpenGLResources.h
--- a/Engine/Source/Runtime/OpenGLDrv/Public/OpenGLResources.h
+++ b/Engine/Source/Runtime/OpenGLDrv/Public/OpenGLResources.h
@@ -876,7 +876,7 @@ public:
 
 
 	/** FRHITexture override.  See FRHITexture::GetNativeResource() */
-	virtual void* GetNativeResource() const OVERRIDE
+	virtual void* GetNativeResource() const override
 	{ 
 		return const_cast<void *>(reinterpret_cast<const void*>(&Resource));
 	}
diff --git a/Engine/Source/Runtime/OpenGLDrv/Public/OpenGLState.h b/Engine/Source/Runtime/OpenGLDrv/Public/OpenGLState.h
--- a/Engine/Source/Runtime/OpenGLDrv/Public/OpenGLState.h
+++ b/Engine/Source/Runtime/OpenGLDrv/Public/OpenGLState.h
@@ -359,14 +359,14 @@ struct FOpenGLContextState : public FOpenGLCommonState
 		FMemory::Memzero(UniformBufferOffsets, sizeof(UniformBufferOffsets));
 	}
 
-	virtual void InitializeResources(int32 NumCombinedTextures, int32 NumComputeUAVUnits) OVERRIDE
+	virtual void InitializeResources(int32 NumCombinedTextures, int32 NumComputeUAVUnits) override
 	{
 		FOpenGLCommonState::InitializeResources(NumCombinedTextures, NumComputeUAVUnits);
 		CachedSamplerStates.Empty(NumCombinedTextures);
 		CachedSamplerStates.AddZeroed(NumCombinedTextures);
 	}
 
-	virtual void CleanupResources() OVERRIDE
+	virtual void CleanupResources() override
 	{
 		CachedSamplerStates.Empty();
 		FOpenGLCommonState::CleanupResources();
@@ -468,7 +468,7 @@ struct FOpenGLRHIState : public FOpenGLCommonState
 		CleanupResources();
 	}
 
-	virtual void InitializeResources(int32 NumCombinedTextures, int32 NumComputeUAVUnits) OVERRIDE
+	virtual void InitializeResources(int32 NumCombinedTextures, int32 NumComputeUAVUnits) override
 	{
 		check(!ShaderParameters);
 		FOpenGLCommonState::InitializeResources(NumCombinedTextures, NumComputeUAVUnits);
@@ -494,7 +494,7 @@ struct FOpenGLRHIState : public FOpenGLCommonState
 		}
 	}
 
-	virtual void CleanupResources() OVERRIDE
+	virtual void CleanupResources() override
 	{
 		delete [] ShaderParameters;
 		ShaderParameters = NULL;
diff --git a/Engine/Source/Runtime/PakFile/Private/IPlatformFilePak.cpp b/Engine/Source/Runtime/PakFile/Private/IPlatformFilePak.cpp
--- a/Engine/Source/Runtime/PakFile/Private/IPlatformFilePak.cpp
+++ b/Engine/Source/Runtime/PakFile/Private/IPlatformFilePak.cpp
@@ -249,7 +249,7 @@ public:
 	{}
 
 	/** Console commands **/
-	virtual bool Exec( UWorld* InWorld, const TCHAR* Cmd, FOutputDevice& Ar ) OVERRIDE
+	virtual bool Exec( UWorld* InWorld, const TCHAR* Cmd, FOutputDevice& Ar ) override
 	{
 		if (FParse::Command(&Cmd, TEXT("Mount")))
 		{
@@ -574,7 +574,7 @@ bool FPakPlatformFile::CopyFile(const TCHAR* To, const TCHAR* From)
 class FPakFileModule : public IPlatformFileModule
 {
 public:
-	virtual IPlatformFile* GetPlatformFile() OVERRIDE
+	virtual IPlatformFile* GetPlatformFile() override
 	{
 		static TScopedPointer<IPlatformFile> AutoDestroySingleton(new FPakPlatformFile());
 		return AutoDestroySingleton.GetOwnedPointer();
diff --git a/Engine/Source/Runtime/PakFile/Private/SignedArchiveReader.h b/Engine/Source/Runtime/PakFile/Private/SignedArchiveReader.h
--- a/Engine/Source/Runtime/PakFile/Private/SignedArchiveReader.h
+++ b/Engine/Source/Runtime/PakFile/Private/SignedArchiveReader.h
@@ -257,16 +257,16 @@ public:
 	virtual ~FSignedArchiveReader();
 
 	// Begin FArchive interface
-	virtual void Serialize(void* Data, int64 Length) OVERRIDE;
-	virtual int64 Tell() OVERRIDE
+	virtual void Serialize(void* Data, int64 Length) override;
+	virtual int64 Tell() override
 	{
 		return PakOffset;
 	}
-	virtual int64 TotalSize() OVERRIDE
+	virtual int64 TotalSize() override
 	{
 		return PakSize;
 	}
-	virtual void Seek(int64 InPos) OVERRIDE
+	virtual void Seek(int64 InPos) override
 	{
 		PakOffset = InPos;
 	}
diff --git a/Engine/Source/Runtime/PakFile/Public/IPlatformFilePak.h b/Engine/Source/Runtime/PakFile/Public/IPlatformFilePak.h
--- a/Engine/Source/Runtime/PakFile/Public/IPlatformFilePak.h
+++ b/Engine/Source/Runtime/PakFile/Public/IPlatformFilePak.h
@@ -572,11 +572,11 @@ public:
 	}
 
 	// BEGIN IFileHandle Interface
-	virtual int64 Tell() OVERRIDE
+	virtual int64 Tell() override
 	{
 		return ReadPos;
 	}
-	virtual bool Seek(int64 NewPosition) OVERRIDE
+	virtual bool Seek(int64 NewPosition) override
 	{
 		if (NewPosition > PakEntry.Size || NewPosition < 0)
 		{
@@ -585,11 +585,11 @@ public:
 		ReadPos = NewPosition;
 		return true;
 	}
-	virtual bool SeekFromEnd(int64 NewPositionRelativeToEnd) OVERRIDE
+	virtual bool SeekFromEnd(int64 NewPositionRelativeToEnd) override
 	{
 		return Seek(PakEntry.Size - NewPositionRelativeToEnd);
 	}
-	virtual bool Read(uint8* Destination, int64 BytesToRead) OVERRIDE
+	virtual bool Read(uint8* Destination, int64 BytesToRead) override
 	{
 		// Check that the file header is OK
 		if (!PakEntry.Verified)
@@ -621,12 +621,12 @@ public:
 			return false;
 		}
 	}
-	virtual bool Write(const uint8* Source, int64 BytesToWrite) OVERRIDE
+	virtual bool Write(const uint8* Source, int64 BytesToWrite) override
 	{
 		// Writing in pak files is not allowed.
 		return false;
 	}
-	virtual int64 Size() OVERRIDE
+	virtual int64 Size() override
 	{
 		return PakEntry.Size;
 	}
@@ -752,15 +752,15 @@ public:
 	 */
 	virtual ~FPakPlatformFile();
 
-	virtual bool ShouldBeUsed(IPlatformFile* Inner, const TCHAR* CmdLine) const OVERRIDE;
-	virtual bool Initialize(IPlatformFile* Inner, const TCHAR* CommandLineParam) OVERRIDE;
+	virtual bool ShouldBeUsed(IPlatformFile* Inner, const TCHAR* CmdLine) const override;
+	virtual bool Initialize(IPlatformFile* Inner, const TCHAR* CommandLineParam) override;
 
-	virtual IPlatformFile* GetLowerLevel() OVERRIDE
+	virtual IPlatformFile* GetLowerLevel() override
 	{
 		return LowerLevel;
 	}
 
-	virtual const TCHAR* GetName() const OVERRIDE
+	virtual const TCHAR* GetName() const override
 	{
 		return FPakPlatformFile::GetTypeName();
 	}
@@ -818,7 +818,7 @@ public:
 	}
 
 	// BEGIN IPlatformFile Interface
-	virtual bool FileExists(const TCHAR* Filename) OVERRIDE
+	virtual bool FileExists(const TCHAR* Filename) override
 	{
 		// Check pak files first.
 		if (FindFileInPakFiles(Filename) != NULL)
@@ -830,7 +830,7 @@ public:
 		return Result;
 	}
 
-	virtual int64 FileSize(const TCHAR* Filename) OVERRIDE
+	virtual int64 FileSize(const TCHAR* Filename) override
 	{
 		// Check pak files first
 		const FPakEntry* FileEntry = FindFileInPakFiles(Filename);
@@ -843,7 +843,7 @@ public:
 		return Result;
 	}
 
-	virtual bool DeleteFile(const TCHAR* Filename) OVERRIDE
+	virtual bool DeleteFile(const TCHAR* Filename) override
 	{
 		// If file exists in pak file it will never get deleted.
 		if (FindFileInPakFiles(Filename) != NULL)
@@ -855,7 +855,7 @@ public:
 		return Result;
 	}
 
-	virtual bool IsReadOnly(const TCHAR* Filename) OVERRIDE
+	virtual bool IsReadOnly(const TCHAR* Filename) override
 	{
 		// Files in pak file are always read-only.
 		if (FindFileInPakFiles(Filename) != NULL)
@@ -867,7 +867,7 @@ public:
 		return Result;
 	}
 
-	virtual bool MoveFile(const TCHAR* To, const TCHAR* From) OVERRIDE
+	virtual bool MoveFile(const TCHAR* To, const TCHAR* From) override
 	{
 		// Files which exist in pak files can't be moved
 		if (FindFileInPakFiles(From) != NULL)
@@ -879,7 +879,7 @@ public:
 		return Result;
 	}
 
-	virtual bool SetReadOnly(const TCHAR* Filename, bool bNewReadOnlyValue) OVERRIDE
+	virtual bool SetReadOnly(const TCHAR* Filename, bool bNewReadOnlyValue) override
 	{
 		// Files in pak file will never change their read-only flag.
 		if (FindFileInPakFiles(Filename) != NULL)
@@ -892,7 +892,7 @@ public:
 		return Result;
 	}
 
-	virtual FDateTime GetTimeStamp(const TCHAR* Filename) OVERRIDE
+	virtual FDateTime GetTimeStamp(const TCHAR* Filename) override
 	{
 		// Check pak files first.
 		FPakFile* PakFile = NULL;
@@ -905,7 +905,7 @@ public:
 		return Result;
 	}
 
-	virtual void SetTimeStamp(const TCHAR* Filename, FDateTime DateTime) OVERRIDE
+	virtual void SetTimeStamp(const TCHAR* Filename, FDateTime DateTime) override
 	{
 		// No modifications allowed on files from pak (although we could theoretically allow this one).
 		if (FindFileInPakFiles(Filename) == NULL)
@@ -914,7 +914,7 @@ public:
 		}
 	}
 
-	virtual FDateTime GetAccessTimeStamp(const TCHAR* Filename) OVERRIDE
+	virtual FDateTime GetAccessTimeStamp(const TCHAR* Filename) override
 	{
 		// AccessTimestamp not yet supported in pak files (although it is possible).
 		FPakFile* PakFile = NULL;
@@ -927,9 +927,9 @@ public:
 		return Result;
 	}
 
-	virtual IFileHandle* OpenRead(const TCHAR* Filename) OVERRIDE;
+	virtual IFileHandle* OpenRead(const TCHAR* Filename) override;
 
-	virtual IFileHandle* OpenWrite(const TCHAR* Filename, bool bAppend = false, bool bAllowRead = false) OVERRIDE
+	virtual IFileHandle* OpenWrite(const TCHAR* Filename, bool bAppend = false, bool bAllowRead = false) override
 	{
 		// No modifications allowed on pak files.
 		if (FindFileInPakFiles(Filename) != NULL)
@@ -940,7 +940,7 @@ public:
 		return LowerLevel->OpenWrite(Filename, bAppend, bAllowRead);
 	}
 
-	virtual bool DirectoryExists(const TCHAR* Directory) OVERRIDE
+	virtual bool DirectoryExists(const TCHAR* Directory) override
 	{
 		// Check pak files first.
 		if (DirectoryExistsInPakFiles(Directory))
@@ -952,13 +952,13 @@ public:
 		return Result;
 	}
 
-	virtual bool CreateDirectory(const TCHAR* Directory) OVERRIDE
+	virtual bool CreateDirectory(const TCHAR* Directory) override
 	{
 		// Directories can be created only under the normal path
 		return LowerLevel->CreateDirectory(Directory);
 	}
 
-	virtual bool DeleteDirectory(const TCHAR* Directory) OVERRIDE
+	virtual bool DeleteDirectory(const TCHAR* Directory) override
 	{
 		// Even if the same directory exists outside of pak files it will never
 		// get truely deleted from pak and will still be reported by Iterate functions.
@@ -1011,7 +1011,7 @@ public:
 		}
 	};
 
-	virtual bool IterateDirectory(const TCHAR* Directory, IPlatformFile::FDirectoryVisitor& Visitor) OVERRIDE
+	virtual bool IterateDirectory(const TCHAR* Directory, IPlatformFile::FDirectoryVisitor& Visitor) override
 	{
 		bool Result = true;
 		TSet<FString> FilesVisitedInPak;
@@ -1066,7 +1066,7 @@ public:
 		return Result;
 	}
 
-	virtual bool IterateDirectoryRecursively(const TCHAR* Directory, IPlatformFile::FDirectoryVisitor& Visitor) OVERRIDE
+	virtual bool IterateDirectoryRecursively(const TCHAR* Directory, IPlatformFile::FDirectoryVisitor& Visitor) override
 	{
 		TSet<FString> FilesVisitedInPak;
 		TArray<FPakFile*> Paks;
@@ -1075,7 +1075,7 @@ public:
 		return IPlatformFile::IterateDirectoryRecursively(Directory, PakVisitor);
 	}
 
-	virtual bool DeleteDirectoryRecursively(const TCHAR* Directory) OVERRIDE
+	virtual bool DeleteDirectoryRecursively(const TCHAR* Directory) override
 	{
 		// Can't delete directories existing in pak files. See DeleteDirectory(..) for more info.
 		if (DirectoryExistsInPakFiles(Directory))
@@ -1086,13 +1086,13 @@ public:
 		return LowerLevel->DeleteDirectoryRecursively(Directory);
 	}
 
-	virtual bool CreateDirectoryTree(const TCHAR* Directory) OVERRIDE
+	virtual bool CreateDirectoryTree(const TCHAR* Directory) override
 	{
 		// Directories can only be created only under the normal path
 		return LowerLevel->CreateDirectoryTree(Directory);
 	}
 
-	virtual bool CopyFile(const TCHAR* To, const TCHAR* From) OVERRIDE;
+	virtual bool CopyFile(const TCHAR* To, const TCHAR* From) override;
 
 	/**
 	 * Converts a filename to a path inside pak file.
@@ -1107,7 +1107,7 @@ public:
 		return RelativeFilename.Mid(Pak->GetMountPoint().Len());
 	}
 
-	FString ConvertToAbsolutePathForExternalAppForRead(const TCHAR* Filename) OVERRIDE
+	FString ConvertToAbsolutePathForExternalAppForRead(const TCHAR* Filename) override
 	{
 		// Check in Pak file first
 		FPakFile* Pak = NULL;
@@ -1121,7 +1121,7 @@ public:
 		}
 	}
 
-	FString ConvertToAbsolutePathForExternalAppForWrite(const TCHAR* Filename) OVERRIDE
+	FString ConvertToAbsolutePathForExternalAppForWrite(const TCHAR* Filename) override
 	{
 		// Check in Pak file first
 		FPakFile* Pak = NULL;
diff --git a/Engine/Source/Runtime/Projects/Private/PluginManager.cpp b/Engine/Source/Runtime/Projects/Private/PluginManager.cpp
--- a/Engine/Source/Runtime/Projects/Private/PluginManager.cpp
+++ b/Engine/Source/Runtime/Projects/Private/PluginManager.cpp
@@ -79,7 +79,7 @@ void FPluginManager::DiscoverAllPlugins()
 					{
 					}
 
-					virtual bool Visit(const TCHAR* FilenameOrDirectory, bool bIsDirectory) OVERRIDE
+					virtual bool Visit(const TCHAR* FilenameOrDirectory, bool bIsDirectory) override
 					{
 						if( bIsDirectory )
 						{
@@ -124,7 +124,7 @@ void FPluginManager::DiscoverAllPlugins()
 								{
 								}
 
-								virtual bool Visit(const TCHAR* FilenameOrDirectory, bool bIsDirectory) OVERRIDE
+								virtual bool Visit(const TCHAR* FilenameOrDirectory, bool bIsDirectory) override
 								{
 									bool bShouldContinue = true;
 									if( !bIsDirectory )
diff --git a/Engine/Source/Runtime/Projects/Private/PluginManager.h b/Engine/Source/Runtime/Projects/Private/PluginManager.h
--- a/Engine/Source/Runtime/Projects/Private/PluginManager.h
+++ b/Engine/Source/Runtime/Projects/Private/PluginManager.h
@@ -87,11 +87,11 @@ public:
 	void SetLoadedFrom(FPluginInfo::ELoadedFrom::Type NewLoadedFrom);
 
 protected:
-	virtual FProjectOrPluginInfo& GetProjectOrPluginInfo() OVERRIDE	{ return PluginInfo; }
-	virtual const FProjectOrPluginInfo& GetProjectOrPluginInfo() const OVERRIDE { return PluginInfo; }
+	virtual FProjectOrPluginInfo& GetProjectOrPluginInfo() override	{ return PluginInfo; }
+	virtual const FProjectOrPluginInfo& GetProjectOrPluginInfo() const override { return PluginInfo; }
 
-	virtual bool PerformAdditionalDeserialization(const TSharedRef< FJsonObject >& FileObject) OVERRIDE;
-	virtual void PerformAdditionalSerialization(const TSharedRef< TJsonWriter<> >& Writer) const OVERRIDE;
+	virtual bool PerformAdditionalDeserialization(const TSharedRef< FJsonObject >& FileObject) override;
+	virtual void PerformAdditionalSerialization(const TSharedRef< TJsonWriter<> >& Writer) const override;
 
 private:
 
@@ -122,13 +122,13 @@ public:
 public:
 
 	/** IPluginManager interface */
-	virtual void LoadModulesForEnabledPlugins( const ELoadingPhase::Type LoadingPhase ) OVERRIDE;
-	virtual void SetRegisterMountPointDelegate( const FRegisterMountPointDelegate& Delegate ) OVERRIDE;
-	virtual bool IsPluginModule( const FName ModuleName ) const OVERRIDE;
-	virtual bool AreEnabledPluginModulesUpToDate() OVERRIDE;
-	virtual TArray< FPluginStatus > QueryStatusForAllPlugins() const OVERRIDE;
-	virtual void SetPluginEnabled( const FString& PluginName, bool bEnabled ) OVERRIDE;
-	virtual bool IsRestartRequired() const OVERRIDE;
+	virtual void LoadModulesForEnabledPlugins( const ELoadingPhase::Type LoadingPhase ) override;
+	virtual void SetRegisterMountPointDelegate( const FRegisterMountPointDelegate& Delegate ) override;
+	virtual bool IsPluginModule( const FName ModuleName ) const override;
+	virtual bool AreEnabledPluginModulesUpToDate() override;
+	virtual TArray< FPluginStatus > QueryStatusForAllPlugins() const override;
+	virtual void SetPluginEnabled( const FString& PluginName, bool bEnabled ) override;
+	virtual bool IsRestartRequired() const override;
 
 private:
 
diff --git a/Engine/Source/Runtime/Projects/Private/ProjectManager.h b/Engine/Source/Runtime/Projects/Private/ProjectManager.h
--- a/Engine/Source/Runtime/Projects/Private/ProjectManager.h
+++ b/Engine/Source/Runtime/Projects/Private/ProjectManager.h
@@ -74,11 +74,11 @@ public:
 	}
 
 protected:
-	virtual FProjectOrPluginInfo& GetProjectOrPluginInfo() OVERRIDE { return ProjectInfo; }
-	virtual const FProjectOrPluginInfo& GetProjectOrPluginInfo() const OVERRIDE { return ProjectInfo; }
+	virtual FProjectOrPluginInfo& GetProjectOrPluginInfo() override { return ProjectInfo; }
+	virtual const FProjectOrPluginInfo& GetProjectOrPluginInfo() const override { return ProjectInfo; }
 
-	virtual bool PerformAdditionalDeserialization(const TSharedRef< FJsonObject >& FileObject) OVERRIDE;
-	virtual void PerformAdditionalSerialization(const TSharedRef< TJsonWriter<> >& Writer) const OVERRIDE;
+	virtual bool PerformAdditionalDeserialization(const TSharedRef< FJsonObject >& FileObject) override;
+	virtual void PerformAdditionalSerialization(const TSharedRef< TJsonWriter<> >& Writer) const override;
 
 private:
 	/** The descriptor of this project that we loaded from disk */
@@ -95,22 +95,22 @@ public:
 	FProjectManager();
 
 	/** IProjectManager interface */
-	virtual bool LoadProjectFile( const FString& ProjectFile ) OVERRIDE;
-	virtual bool LoadModulesForProject( const ELoadingPhase::Type LoadingPhase ) OVERRIDE;
-	virtual bool AreProjectModulesUpToDate( ) OVERRIDE;
-	virtual const FString& GetAutoLoadProjectFileName() OVERRIDE;
-	virtual const FString& NonStaticGetProjectFileExtension() OVERRIDE;
-	virtual bool GenerateNewProjectFile(const FString& NewProjectFilename, const TArray<FString>& ModuleNames, const FString& EngineIdentifier, FText& OutFailReason) OVERRIDE;
-	virtual bool DuplicateProjectFile(const FString& SourceProjectFilename, const FString& NewProjectFilename, const FString& EngineIdentifier, FText& OutFailReason) OVERRIDE;
-	virtual bool UpdateLoadedProjectFileToCurrent(const TArray<FString>* StartupModuleNames, const FString& EngineIdentifier, FText& OutFailReason) OVERRIDE;
-	virtual bool SignSampleProject(const FString& FilePath, const FString& Category, FText& OutFailReason) OVERRIDE;
-	virtual bool QueryStatusForProject(const FString& FilePath, FProjectStatus& OutProjectStatus) const OVERRIDE;
-	virtual bool QueryStatusForCurrentProject(FProjectStatus& OutProjectStatus) const OVERRIDE;
-	virtual void UpdateSupportedTargetPlatformsForProject(const FString& FilePath, const FName& InPlatformName, const bool bIsSupported) OVERRIDE;
-	virtual void UpdateSupportedTargetPlatformsForCurrentProject(const FName& InPlatformName, const bool bIsSupported) OVERRIDE;
-	virtual void ClearSupportedTargetPlatformsForProject(const FString& FilePath) OVERRIDE;
-	virtual void ClearSupportedTargetPlatformsForCurrentProject() OVERRIDE;
-	virtual FOnTargetPlatformsForCurrentProjectChangedEvent& OnTargetPlatformsForCurrentProjectChanged() OVERRIDE { return OnTargetPlatformsForCurrentProjectChangedEvent; }
+	virtual bool LoadProjectFile( const FString& ProjectFile ) override;
+	virtual bool LoadModulesForProject( const ELoadingPhase::Type LoadingPhase ) override;
+	virtual bool AreProjectModulesUpToDate( ) override;
+	virtual const FString& GetAutoLoadProjectFileName() override;
+	virtual const FString& NonStaticGetProjectFileExtension() override;
+	virtual bool GenerateNewProjectFile(const FString& NewProjectFilename, const TArray<FString>& ModuleNames, const FString& EngineIdentifier, FText& OutFailReason) override;
+	virtual bool DuplicateProjectFile(const FString& SourceProjectFilename, const FString& NewProjectFilename, const FString& EngineIdentifier, FText& OutFailReason) override;
+	virtual bool UpdateLoadedProjectFileToCurrent(const TArray<FString>* StartupModuleNames, const FString& EngineIdentifier, FText& OutFailReason) override;
+	virtual bool SignSampleProject(const FString& FilePath, const FString& Category, FText& OutFailReason) override;
+	virtual bool QueryStatusForProject(const FString& FilePath, FProjectStatus& OutProjectStatus) const override;
+	virtual bool QueryStatusForCurrentProject(FProjectStatus& OutProjectStatus) const override;
+	virtual void UpdateSupportedTargetPlatformsForProject(const FString& FilePath, const FName& InPlatformName, const bool bIsSupported) override;
+	virtual void UpdateSupportedTargetPlatformsForCurrentProject(const FName& InPlatformName, const bool bIsSupported) override;
+	virtual void ClearSupportedTargetPlatformsForProject(const FString& FilePath) override;
+	virtual void ClearSupportedTargetPlatformsForCurrentProject() override;
+	virtual FOnTargetPlatformsForCurrentProjectChangedEvent& OnTargetPlatformsForCurrentProjectChanged() override { return OnTargetPlatformsForCurrentProjectChangedEvent; }
 
 private:
 	static void QueryStatusForProjectImpl(const FProject& Project, const FString& FilePath, FProjectStatus& OutProjectStatus);
diff --git a/Engine/Source/Runtime/RHI/Public/RHIResources.h b/Engine/Source/Runtime/RHI/Public/RHIResources.h
--- a/Engine/Source/Runtime/RHI/Public/RHIResources.h
+++ b/Engine/Source/Runtime/RHI/Public/RHIResources.h
@@ -376,7 +376,7 @@ public:
 		: FRHITexture(0,0,PF_Unknown,0)
 	{}
 
-	virtual FRHITextureReference* GetTextureReference() OVERRIDE { return this; }
+	virtual FRHITextureReference* GetTextureReference() override { return this; }
 	inline FRHITexture* GetReferencedTexture() const { return ReferencedTexture.GetReference(); }
 
 protected:
diff --git a/Engine/Source/Runtime/RenderCore/Private/RenderResource.cpp b/Engine/Source/Runtime/RenderCore/Private/RenderResource.cpp
--- a/Engine/Source/Runtime/RenderCore/Private/RenderResource.cpp
+++ b/Engine/Source/Runtime/RenderCore/Private/RenderResource.cpp
@@ -168,7 +168,7 @@ public:
 	}
 
 	// FRenderResource interface.
-	virtual void InitRHI() OVERRIDE
+	virtual void InitRHI() override
 	{
 		check(!IsValidRef(VertexBufferRHI));
 		FRHIResourceCreateInfo CreateInfo;
@@ -177,14 +177,14 @@ public:
 		AllocatedByteCount = 0;
 	}
 
-	virtual void ReleaseRHI() OVERRIDE
+	virtual void ReleaseRHI() override
 	{
 		FVertexBuffer::ReleaseRHI();
 		MappedBuffer = NULL;
 		AllocatedByteCount = 0;
 	}
 
-	virtual FString GetFriendlyName() const OVERRIDE
+	virtual FString GetFriendlyName() const override
 	{
 		return TEXT("FDynamicVertexBuffer");
 	}
@@ -350,7 +350,7 @@ public:
 	}
 
 	// FRenderResource interface.
-	virtual void InitRHI() OVERRIDE
+	virtual void InitRHI() override
 	{
 		check(!IsValidRef(IndexBufferRHI));
 		FRHIResourceCreateInfo CreateInfo;
@@ -359,14 +359,14 @@ public:
 		AllocatedByteCount = 0;
 	}
 
-	virtual void ReleaseRHI() OVERRIDE
+	virtual void ReleaseRHI() override
 	{
 		FIndexBuffer::ReleaseRHI();
 		MappedBuffer = NULL;
 		AllocatedByteCount = 0;
 	}
 
-	virtual FString GetFriendlyName() const OVERRIDE
+	virtual FString GetFriendlyName() const override
 	{
 		return TEXT("FDynamicIndexBuffer");
 	}
diff --git a/Engine/Source/Runtime/RenderCore/Private/RenderUtils.cpp b/Engine/Source/Runtime/RenderCore/Private/RenderUtils.cpp
--- a/Engine/Source/Runtime/RenderCore/Private/RenderUtils.cpp
+++ b/Engine/Source/Runtime/RenderCore/Private/RenderUtils.cpp
@@ -199,7 +199,7 @@ public:
 	/**
 	 * Returns a pointer to the bulk data.
 	 */
-	virtual const void* GetResourceBulkData() const OVERRIDE
+	virtual const void* GetResourceBulkData() const override
 	{
 		return &Color;
 	}
@@ -207,7 +207,7 @@ public:
 	/** 
 	 * @return size of resource memory
 	 */
-	virtual uint32 GetResourceBulkDataSize() const OVERRIDE
+	virtual uint32 GetResourceBulkDataSize() const override
 	{
 		return sizeof(Color);
 	}
@@ -215,7 +215,7 @@ public:
 	/**
 	 * Free memory after it has been used to initialize RHI resource 
 	 */
-	virtual void Discard() OVERRIDE
+	virtual void Discard() override
 	{
 	}
 
diff --git a/Engine/Source/Runtime/Renderer/Private/CompositionLighting/PostProcessAmbient.h b/Engine/Source/Runtime/Renderer/Private/CompositionLighting/PostProcessAmbient.h
--- a/Engine/Source/Runtime/Renderer/Private/CompositionLighting/PostProcessAmbient.h
+++ b/Engine/Source/Runtime/Renderer/Private/CompositionLighting/PostProcessAmbient.h
@@ -18,6 +18,6 @@ public:
 	// interface FRenderingCompositePass ---------
 	virtual const TCHAR* GetDebugName() { return TEXT("FRCPassPostProcessAmbient"); }
 	virtual void Process(FRenderingCompositePassContext& Context);
-	virtual void Release() OVERRIDE { delete this; }
+	virtual void Release() override { delete this; }
 	virtual FPooledRenderTargetDesc ComputeOutputDesc(EPassOutputId InPassOutputId) const;
 };
diff --git a/Engine/Source/Runtime/Renderer/Private/CompositionLighting/PostProcessAmbientOcclusion.h b/Engine/Source/Runtime/Renderer/Private/CompositionLighting/PostProcessAmbientOcclusion.h
--- a/Engine/Source/Runtime/Renderer/Private/CompositionLighting/PostProcessAmbientOcclusion.h
+++ b/Engine/Source/Runtime/Renderer/Private/CompositionLighting/PostProcessAmbientOcclusion.h
@@ -18,7 +18,7 @@ public:
 
 	// interface FRenderingCompositePass ---------
 	virtual void Process(FRenderingCompositePassContext& Context);
-	virtual void Release() OVERRIDE { delete this; }
+	virtual void Release() override { delete this; }
 	virtual FPooledRenderTargetDesc ComputeOutputDesc(EPassOutputId InPassOutputId) const;
 
 private:
@@ -43,7 +43,7 @@ public:
 
 	// interface FRenderingCompositePass ---------
 	virtual void Process(FRenderingCompositePassContext& Context);
-	virtual void Release() OVERRIDE { delete this; }
+	virtual void Release() override { delete this; }
 	virtual FPooledRenderTargetDesc ComputeOutputDesc(EPassOutputId InPassOutputId) const;
 
 private:
@@ -61,7 +61,7 @@ class FRCPassPostProcessBasePassAO : public TRenderingCompositePassBase<0, 1>
 public:
 	// interface FRenderingCompositePass ---------
 	virtual void Process(FRenderingCompositePassContext& Context);
-	virtual void Release() OVERRIDE { delete this; }
+	virtual void Release() override { delete this; }
 	virtual FPooledRenderTargetDesc ComputeOutputDesc(EPassOutputId InPassOutputId) const;
 };
 
diff --git a/Engine/Source/Runtime/Renderer/Private/CompositionLighting/PostProcessDeferredDecals.cpp b/Engine/Source/Runtime/Renderer/Private/CompositionLighting/PostProcessDeferredDecals.cpp
--- a/Engine/Source/Runtime/Renderer/Private/CompositionLighting/PostProcessDeferredDecals.cpp
+++ b/Engine/Source/Runtime/Renderer/Private/CompositionLighting/PostProcessDeferredDecals.cpp
@@ -535,7 +535,7 @@ public:
 	/**
 	* Initialize the RHI for this rendering resource
 	*/
-	void InitRHI() OVERRIDE
+	void InitRHI() override
 	{
 		const int32 NumVerts = 8;
 		TResourceArray<FVector4, VERTEXBUFFER_ALIGNMENT> Verts;
@@ -576,7 +576,7 @@ public:
 	/**
 	* Initialize the RHI for this rendering resource
 	*/
-	void InitRHI() OVERRIDE
+	void InitRHI() override
 	{
 		TResourceArray<uint16, INDEXBUFFER_ALIGNMENT> Indices;
 		
diff --git a/Engine/Source/Runtime/Renderer/Private/CompositionLighting/PostProcessDeferredDecals.h b/Engine/Source/Runtime/Renderer/Private/CompositionLighting/PostProcessDeferredDecals.h
--- a/Engine/Source/Runtime/Renderer/Private/CompositionLighting/PostProcessDeferredDecals.h
+++ b/Engine/Source/Runtime/Renderer/Private/CompositionLighting/PostProcessDeferredDecals.h
@@ -18,7 +18,7 @@ public:
 
 	// interface FRenderingCompositePass ---------
 	virtual void Process(FRenderingCompositePassContext& Context);
-	virtual void Release() OVERRIDE { delete this; }
+	virtual void Release() override { delete this; }
 	virtual FPooledRenderTargetDesc ComputeOutputDesc(EPassOutputId InPassOutputId) const;
 
 private:
diff --git a/Engine/Source/Runtime/Renderer/Private/DeferredShadingRenderer.h b/Engine/Source/Runtime/Renderer/Private/DeferredShadingRenderer.h
--- a/Engine/Source/Runtime/Renderer/Private/DeferredShadingRenderer.h
+++ b/Engine/Source/Runtime/Renderer/Private/DeferredShadingRenderer.h
@@ -74,10 +74,10 @@ public:
     void RenderFinish();
 
 	/** Renders the view family. */
-	virtual void Render() OVERRIDE;
+	virtual void Render() override;
 
 	/** Render the view family's hit proxies. */
-	virtual void RenderHitProxies() OVERRIDE;
+	virtual void RenderHitProxies() override;
 
 #if !(UE_BUILD_SHIPPING || UE_BUILD_TEST)
 	void RenderVisualizeTexturePool();
diff --git a/Engine/Source/Runtime/Renderer/Private/DynamicPrimitiveDrawing.h b/Engine/Source/Runtime/Renderer/Private/DynamicPrimitiveDrawing.h
--- a/Engine/Source/Runtime/Renderer/Private/DynamicPrimitiveDrawing.h
+++ b/Engine/Source/Runtime/Renderer/Private/DynamicPrimitiveDrawing.h
@@ -50,13 +50,13 @@ public:
 	void SetPrimitive(const FPrimitiveSceneProxy* NewPrimitiveSceneProxy);
 
 	// FPrimitiveDrawInterface interface.
-	virtual bool IsHitTesting() OVERRIDE;
-	virtual void SetHitProxy(HHitProxy* HitProxy) OVERRIDE;
-	virtual bool IsRenderingSelectionOutline() const OVERRIDE { return bIsSelectionOutlineRendering; }
-	virtual void RegisterDynamicResource(FDynamicPrimitiveResource* DynamicResource) OVERRIDE;
-	virtual bool IsMaterialIgnored(const FMaterialRenderProxy* MaterialRenderProxy, ERHIFeatureLevel::Type InFeatureLevel) const OVERRIDE;
-	virtual int32 DrawMesh(const FMeshBatch& Mesh) OVERRIDE;
-	virtual void AddReserveLines(uint8 DepthPriorityGroup, int32 NumLines, bool bDepthBiased = false) OVERRIDE;
+	virtual bool IsHitTesting() override;
+	virtual void SetHitProxy(HHitProxy* HitProxy) override;
+	virtual bool IsRenderingSelectionOutline() const override { return bIsSelectionOutlineRendering; }
+	virtual void RegisterDynamicResource(FDynamicPrimitiveResource* DynamicResource) override;
+	virtual bool IsMaterialIgnored(const FMaterialRenderProxy* MaterialRenderProxy, ERHIFeatureLevel::Type InFeatureLevel) const override;
+	virtual int32 DrawMesh(const FMeshBatch& Mesh) override;
+	virtual void AddReserveLines(uint8 DepthPriorityGroup, int32 NumLines, bool bDepthBiased = false) override;
 	virtual void DrawSprite(
 		const FVector& Position,
 		float SizeX,
@@ -69,7 +69,7 @@ public:
 		float V,
 		float VL,
 		uint8 BlendMode = SE_BLEND_Masked
-		) OVERRIDE;
+		) override;
 	virtual void DrawLine(
 		const FVector& Start,
 		const FVector& End,
@@ -78,13 +78,13 @@ public:
 		float Thickness = 0.0f,
 		float DepthBias = 0.0f,
 		bool bScreenSpace = false
-		) OVERRIDE;
+		) override;
 	virtual void DrawPoint(
 		const FVector& Position,
 		const FLinearColor& Color,
 		float PointSize,
 		uint8 DepthPriorityGroup
-		) OVERRIDE;
+		) override;
 
 	// Accessors.
 	bool IsPreFog() const
@@ -181,10 +181,10 @@ public:
 	FViewElementPDI(FViewInfo* InViewInfo,FHitProxyConsumer* InHitProxyConsumer);
 
 	// FPrimitiveDrawInterface interface.
-	virtual bool IsHitTesting() OVERRIDE;
-	virtual void SetHitProxy(HHitProxy* HitProxy) OVERRIDE;
-	virtual void RegisterDynamicResource(FDynamicPrimitiveResource* DynamicResource) OVERRIDE;
-	virtual void AddReserveLines(uint8 DepthPriorityGroup, int32 NumLines, bool bDepthBiased = false) OVERRIDE;
+	virtual bool IsHitTesting() override;
+	virtual void SetHitProxy(HHitProxy* HitProxy) override;
+	virtual void RegisterDynamicResource(FDynamicPrimitiveResource* DynamicResource) override;
+	virtual void AddReserveLines(uint8 DepthPriorityGroup, int32 NumLines, bool bDepthBiased = false) override;
 	virtual void DrawSprite(
 		const FVector& Position,
 		float SizeX,
@@ -197,7 +197,7 @@ public:
 		float V,
 		float VL,
 		uint8 BlendMode = SE_BLEND_Masked
-		) OVERRIDE;
+		) override;
 	virtual void DrawLine(
 		const FVector& Start,
 		const FVector& End,
@@ -206,14 +206,14 @@ public:
 		float Thickness = 0.0f,
 		float DepthBias = 0.0f,
 		bool bScreenSpace = false
-		) OVERRIDE;
+		) override;
 	virtual void DrawPoint(
 		const FVector& Position,
 		const FLinearColor& Color,
 		float PointSize,
 		uint8 DepthPriorityGroup
-		) OVERRIDE;
-	virtual int32 DrawMesh(const FMeshBatch& Mesh) OVERRIDE;
+		) override;
+	virtual int32 DrawMesh(const FMeshBatch& Mesh) override;
 
 private:
 	FViewInfo* ViewInfo;
diff --git a/Engine/Source/Runtime/Renderer/Private/LightFunctionRendering.cpp b/Engine/Source/Runtime/Renderer/Private/LightFunctionRendering.cpp
--- a/Engine/Source/Runtime/Renderer/Private/LightFunctionRendering.cpp
+++ b/Engine/Source/Runtime/Renderer/Private/LightFunctionRendering.cpp
@@ -51,7 +51,7 @@ public:
 	}
 
 	// Begin FShader interface
-	virtual bool Serialize(FArchive& Ar) OVERRIDE
+	virtual bool Serialize(FArchive& Ar) override
 	{
 		bool bShaderHasOutdatedParameters = FMaterialShader::Serialize(Ar);
 		Ar << StencilingGeometryParameters;
diff --git a/Engine/Source/Runtime/Renderer/Private/LightPropagationVolume.cpp b/Engine/Source/Runtime/Renderer/Private/LightPropagationVolume.cpp
--- a/Engine/Source/Runtime/Renderer/Private/LightPropagationVolume.cpp
+++ b/Engine/Source/Runtime/Renderer/Private/LightPropagationVolume.cpp
@@ -135,7 +135,7 @@ public:
 		OutEnvironment.SetDefine( TEXT("LPV_GV_VOLUME_TEXTURE"),	(uint32)LPV_GV_VOLUME_TEXTURE );
 	}
 	// Serialization
-	virtual bool Serialize( FArchive& Ar ) OVERRIDE
+	virtual bool Serialize( FArchive& Ar ) override
 	{
 		bool bShaderHasOutdatedParameters = FGlobalShader::Serialize( Ar );
 
@@ -392,7 +392,7 @@ public:
 
 	explicit FLpvClearCS( const ShaderMetaType::CompiledShaderInitializerType& Initializer ) : FLpvWriteShaderCSBase(Initializer)		{	}
 
-	virtual bool Serialize( FArchive& Ar ) OVERRIDE			{ return FLpvWriteShaderCSBase::Serialize( Ar ); }
+	virtual bool Serialize( FArchive& Ar ) override			{ return FLpvWriteShaderCSBase::Serialize( Ar ); }
 };
 IMPLEMENT_SHADER_TYPE(,FLpvClearCS,TEXT("LPVClear"),TEXT("CSClear"),SF_Compute);
 
@@ -417,7 +417,7 @@ public:
 
 	explicit FLpvClearGeometryVolumeCS( const ShaderMetaType::CompiledShaderInitializerType& Initializer ) : FLpvWriteShaderCSBase(Initializer)		{	}
 
-	virtual bool Serialize( FArchive& Ar ) OVERRIDE			{ return FLpvWriteShaderCSBase::Serialize( Ar ); }
+	virtual bool Serialize( FArchive& Ar ) override			{ return FLpvWriteShaderCSBase::Serialize( Ar ); }
 };
 IMPLEMENT_SHADER_TYPE(,FLpvClearGeometryVolumeCS,TEXT("LPVClear"),TEXT("CSClearGeometryVolume"),SF_Compute);
 
@@ -442,7 +442,7 @@ public:
 
 	explicit FLpvClearListsCS( const ShaderMetaType::CompiledShaderInitializerType& Initializer ) : FLpvWriteShaderCSBase(Initializer)		{	}
 
-	virtual bool Serialize( FArchive& Ar ) OVERRIDE			{ return FLpvWriteShaderCSBase::Serialize( Ar ); }
+	virtual bool Serialize( FArchive& Ar ) override			{ return FLpvWriteShaderCSBase::Serialize( Ar ); }
 };
 IMPLEMENT_SHADER_TYPE(,FLpvClearListsCS,TEXT("LPVClearLists"),TEXT("CSClearLists"),SF_Compute);
 
@@ -505,7 +505,7 @@ public:
 	}
 
 
-	virtual bool Serialize( FArchive& Ar ) OVERRIDE			
+	virtual bool Serialize( FArchive& Ar ) override			
 	{ 
 		bool rv = FLpvWriteShaderCSBase::Serialize( Ar ); 
 		Ar << RsmDiffuseTexture;
@@ -555,7 +555,7 @@ public:
 
 	explicit FLpvInject_AccumulateVplListsCS( const ShaderMetaType::CompiledShaderInitializerType& Initializer ) : FLpvWriteShaderCSBase(Initializer)		{	}
 
-	virtual bool Serialize( FArchive& Ar ) OVERRIDE			{ return FLpvWriteShaderCSBase::Serialize( Ar ); }
+	virtual bool Serialize( FArchive& Ar ) override			{ return FLpvWriteShaderCSBase::Serialize( Ar ); }
 };
 IMPLEMENT_SHADER_TYPE(,FLpvInject_AccumulateVplListsCS,TEXT("LPVInject_AccumulateVplLists"),TEXT("CSAccumulateVplLists"),SF_Compute);
 
@@ -580,7 +580,7 @@ public:
 
 	explicit FLpvBuildGeometryVolumeCS( const ShaderMetaType::CompiledShaderInitializerType& Initializer ) : FLpvWriteShaderCSBase(Initializer)		{	}
 
-	virtual bool Serialize( FArchive& Ar ) OVERRIDE			{ return FLpvWriteShaderCSBase::Serialize( Ar ); }
+	virtual bool Serialize( FArchive& Ar ) override			{ return FLpvWriteShaderCSBase::Serialize( Ar ); }
 };
 IMPLEMENT_SHADER_TYPE(,FLpvBuildGeometryVolumeCS,TEXT("LPVBuildGeometryVolume"),TEXT("CSBuildGeometryVolume"),SF_Compute);
 
@@ -618,7 +618,7 @@ public:
 
 	explicit TLpvPropagateCS( const ShaderMetaType::CompiledShaderInitializerType& Initializer ) : FLpvWriteShaderCSBase(Initializer)		{	}
 
-	virtual bool Serialize( FArchive& Ar ) OVERRIDE			{ return FLpvWriteShaderCSBase::Serialize( Ar ); }
+	virtual bool Serialize( FArchive& Ar ) override			{ return FLpvWriteShaderCSBase::Serialize( Ar ); }
 };
 
 IMPLEMENT_SHADER_TYPE(template<>,TLpvPropagateCS<0>,													TEXT("LPVPropagate"),TEXT("CSPropogate"),SF_Compute);
@@ -671,7 +671,7 @@ public:
 
 	explicit FLpvInjectShader_Base( const ShaderMetaType::CompiledShaderInitializerType& Initializer ) : FLpvWriteShaderCSBase(Initializer)		{	}
 
-	virtual bool Serialize( FArchive& Ar ) OVERRIDE			{ return FLpvWriteShaderCSBase::Serialize( Ar ); }
+	virtual bool Serialize( FArchive& Ar ) override			{ return FLpvWriteShaderCSBase::Serialize( Ar ); }
 };
 
 // ----------------------------------------------------------------------------
@@ -696,7 +696,7 @@ public:
 
 	explicit TLpvInject_PointLightCS( const ShaderMetaType::CompiledShaderInitializerType& Initializer ) : FLpvInjectShader_Base(Initializer)		{	}
 
-	virtual bool Serialize( FArchive& Ar ) OVERRIDE			{ return FLpvInjectShader_Base::Serialize( Ar ); }
+	virtual bool Serialize( FArchive& Ar ) override			{ return FLpvInjectShader_Base::Serialize( Ar ); }
 };
 IMPLEMENT_SHADER_TYPE(template<>,TLpvInject_PointLightCS<0>,TEXT("LPVDirectLightInject"),TEXT("CSPointLightInject_ListGenCS"),SF_Compute);
 IMPLEMENT_SHADER_TYPE(template<>,TLpvInject_PointLightCS<INJECT_SHADOWED>,TEXT("LPVDirectLightInject"),TEXT("CSPointLightInject_ListGenCS"),SF_Compute);
diff --git a/Engine/Source/Runtime/Renderer/Private/LightPropagationVolumeVisualisation.cpp b/Engine/Source/Runtime/Renderer/Private/LightPropagationVolumeVisualisation.cpp
--- a/Engine/Source/Runtime/Renderer/Private/LightPropagationVolumeVisualisation.cpp
+++ b/Engine/Source/Runtime/Renderer/Private/LightPropagationVolumeVisualisation.cpp
@@ -49,7 +49,7 @@ public:
 
 	FLpvVisualiseGS()																												{}
 	explicit FLpvVisualiseGS( const ShaderMetaType::CompiledShaderInitializerType& Initializer ) : FLpvVisualiseBase(Initializer)	{}
-	virtual bool Serialize( FArchive& Ar ) OVERRIDE																					{ return FLpvVisualiseBase::Serialize( Ar ); }
+	virtual bool Serialize( FArchive& Ar ) override																					{ return FLpvVisualiseBase::Serialize( Ar ); }
 
 	//@todo-rco: Remove this when reenabling for OpenGL
 	static bool ShouldCache( EShaderPlatform Platform )		{ return IsFeatureLevelSupported(Platform, ERHIFeatureLevel::SM5) && !IsOpenGLPlatform(Platform); }
@@ -71,7 +71,7 @@ public:
 
 	FLpvVisualiseVS()	{	}
 	explicit FLpvVisualiseVS( const ShaderMetaType::CompiledShaderInitializerType& Initializer ) : FLpvVisualiseBase(Initializer) {}
-	virtual bool Serialize( FArchive& Ar ) OVERRIDE																					{ return FLpvVisualiseBase::Serialize( Ar ); }
+	virtual bool Serialize( FArchive& Ar ) override																					{ return FLpvVisualiseBase::Serialize( Ar ); }
 
 	//@todo-rco: Remove this when reenabling for OpenGL
 	static bool ShouldCache( EShaderPlatform Platform )		{ return IsFeatureLevelSupported(Platform, ERHIFeatureLevel::SM5) && !IsOpenGLPlatform(Platform); }
@@ -168,7 +168,7 @@ public:
 
 
 	// Serialization
-	virtual bool Serialize( FArchive& Ar ) OVERRIDE
+	virtual bool Serialize( FArchive& Ar ) override
 	{
 		bool bShaderHasOutdatedParameters = FLpvVisualiseBase::Serialize( Ar );
 
diff --git a/Engine/Source/Runtime/Renderer/Private/PostProcess/PostProcessBokehDOF.cpp b/Engine/Source/Runtime/Renderer/Private/PostProcess/PostProcessBokehDOF.cpp
--- a/Engine/Source/Runtime/Renderer/Private/PostProcess/PostProcessBokehDOF.cpp
+++ b/Engine/Source/Runtime/Renderer/Private/PostProcess/PostProcessBokehDOF.cpp
@@ -27,7 +27,7 @@ template<EBokehIndexStyle DOFIndexStyle>
 class TBokehIndexBuffer : public FIndexBuffer
 {
 public:
-	virtual void InitRHI() OVERRIDE
+	virtual void InitRHI() override
 	{
 		const uint32 Size = sizeof(uint16) * 6 * 8;
 		const uint32 Stride = sizeof(uint16);
diff --git a/Engine/Source/Runtime/Renderer/Private/PostProcess/PostProcessBokehDOF.h b/Engine/Source/Runtime/Renderer/Private/PostProcess/PostProcessBokehDOF.h
--- a/Engine/Source/Runtime/Renderer/Private/PostProcess/PostProcessBokehDOF.h
+++ b/Engine/Source/Runtime/Renderer/Private/PostProcess/PostProcessBokehDOF.h
@@ -15,7 +15,7 @@ class FRCPassPostProcessVisualizeDOF : public TRenderingCompositePassBase<1, 1>
 public:
 	// interface FRenderingCompositePass ---------
 	virtual void Process(FRenderingCompositePassContext& Context);
-	virtual void Release() OVERRIDE { delete this; }
+	virtual void Release() override { delete this; }
 	virtual FPooledRenderTargetDesc ComputeOutputDesc(EPassOutputId InPassOutputId) const;
 };
 
@@ -27,7 +27,7 @@ class FRCPassPostProcessBokehDOFSetup : public TRenderingCompositePassBase<2, 1>
 public:
 	// interface FRenderingCompositePass ---------
 	virtual void Process(FRenderingCompositePassContext& Context);
-	virtual void Release() OVERRIDE { delete this; }
+	virtual void Release() override { delete this; }
 	virtual FPooledRenderTargetDesc ComputeOutputDesc(EPassOutputId InPassOutputId) const;
 };
 
@@ -41,7 +41,7 @@ public:
 	// interface FRenderingCompositePass ---------
 
 	virtual void Process(FRenderingCompositePassContext& Context);
-	virtual void Release() OVERRIDE { delete this; }
+	virtual void Release() override { delete this; }
 	virtual FPooledRenderTargetDesc ComputeOutputDesc(EPassOutputId InPassOutputId) const;
 
 	static void ComputeDepthOfFieldParams(const FRenderingCompositePassContext& Context, FVector4 Out[2]);
diff --git a/Engine/Source/Runtime/Renderer/Private/PostProcess/PostProcessDOF.h b/Engine/Source/Runtime/Renderer/Private/PostProcess/PostProcessDOF.h
--- a/Engine/Source/Runtime/Renderer/Private/PostProcess/PostProcessDOF.h
+++ b/Engine/Source/Runtime/Renderer/Private/PostProcess/PostProcessDOF.h
@@ -21,7 +21,7 @@ public:
 	// interface FRenderingCompositePass ---------
 
 	virtual void Process(FRenderingCompositePassContext& Context);
-	virtual void Release() OVERRIDE { delete this; }
+	virtual void Release() override { delete this; }
 	virtual FPooledRenderTargetDesc ComputeOutputDesc(EPassOutputId InPassOutputId) const;
 
 private:
@@ -41,7 +41,7 @@ public:
 
 	// interface FRenderingCompositePass ---------
 	virtual void Process(FRenderingCompositePassContext& Context);
-	virtual void Release() OVERRIDE { delete this; }
+	virtual void Release() override { delete this; }
 	virtual FPooledRenderTargetDesc ComputeOutputDesc(EPassOutputId InPassOutputId) const;
 
 private:
diff --git a/Engine/Source/Runtime/Renderer/Private/PostProcess/PostProcessInput.h b/Engine/Source/Runtime/Renderer/Private/PostProcess/PostProcessInput.h
--- a/Engine/Source/Runtime/Renderer/Private/PostProcess/PostProcessInput.h
+++ b/Engine/Source/Runtime/Renderer/Private/PostProcess/PostProcessInput.h
@@ -18,7 +18,7 @@ public:
 	// interface FRenderingCompositePass ---------
 
 	virtual void Process(FRenderingCompositePassContext& Context);
-	virtual void Release() OVERRIDE { delete this; }
+	virtual void Release() override { delete this; }
 	virtual FPooledRenderTargetDesc ComputeOutputDesc(EPassOutputId InPassOutputId) const;
 
 protected:
@@ -39,5 +39,5 @@ public:
 	// interface FRenderingCompositePass ---------
 
 	virtual void Process(FRenderingCompositePassContext& Context);
-	virtual void Release() OVERRIDE { delete this; }
+	virtual void Release() override { delete this; }
 };
diff --git a/Engine/Source/Runtime/Renderer/Private/PostProcess/PostProcessMotionBlur.h b/Engine/Source/Runtime/Renderer/Private/PostProcess/PostProcessMotionBlur.h
--- a/Engine/Source/Runtime/Renderer/Private/PostProcess/PostProcessMotionBlur.h
+++ b/Engine/Source/Runtime/Renderer/Private/PostProcess/PostProcessMotionBlur.h
@@ -20,7 +20,7 @@ public:
 	// interface FRenderingCompositePass ---------
 
 	virtual void Process(FRenderingCompositePassContext& Context);
-	virtual void Release() OVERRIDE { delete this; }
+	virtual void Release() override { delete this; }
 	virtual FPooledRenderTargetDesc ComputeOutputDesc(EPassOutputId InPassOutputId) const;
 };
 
@@ -38,7 +38,7 @@ public:
 
 	// interface FRenderingCompositePass ---------
 	virtual void Process(FRenderingCompositePassContext& Context);
-	virtual void Release() OVERRIDE { delete this; }
+	virtual void Release() override { delete this; }
 	virtual FPooledRenderTargetDesc ComputeOutputDesc(EPassOutputId InPassOutputId) const;
 
 	// 1:low, 2:medium, 3:high, 4: very high
@@ -54,7 +54,7 @@ class FRCPassPostProcessMotionBlurRecombine : public TRenderingCompositePassBase
 public:
 	// interface FRenderingCompositePass ---------
 	virtual void Process(FRenderingCompositePassContext& Context);
-	virtual void Release() OVERRIDE { delete this; }
+	virtual void Release() override { delete this; }
 	virtual FPooledRenderTargetDesc ComputeOutputDesc(EPassOutputId InPassOutputId) const;
 };
 
@@ -68,6 +68,6 @@ class FRCPassPostProcessVisualizeMotionBlur : public TRenderingCompositePassBase
 public:
 	// interface FRenderingCompositePass ---------
 	virtual void Process(FRenderingCompositePassContext& Context);
-	virtual void Release() OVERRIDE { delete this; }
+	virtual void Release() override { delete this; }
 	virtual FPooledRenderTargetDesc ComputeOutputDesc(EPassOutputId InPassOutputId) const;
 };
diff --git a/Engine/Source/Runtime/Renderer/Private/PostProcess/PostProcessSelectionOutline.h b/Engine/Source/Runtime/Renderer/Private/PostProcess/PostProcessSelectionOutline.h
--- a/Engine/Source/Runtime/Renderer/Private/PostProcess/PostProcessSelectionOutline.h
+++ b/Engine/Source/Runtime/Renderer/Private/PostProcess/PostProcessSelectionOutline.h
@@ -14,9 +14,9 @@ class FRCPassPostProcessSelectionOutlineColor : public TRenderingCompositePassBa
 {
 public:
 	// interface FRenderingCompositePass ---------
-	virtual void Process(FRenderingCompositePassContext& Context) OVERRIDE;
-	virtual void Release() OVERRIDE { delete this; }
-	virtual FPooledRenderTargetDesc ComputeOutputDesc(EPassOutputId InPassOutputId) const OVERRIDE;
+	virtual void Process(FRenderingCompositePassContext& Context) override;
+	virtual void Release() override { delete this; }
+	virtual FPooledRenderTargetDesc ComputeOutputDesc(EPassOutputId InPassOutputId) const override;
 };
 
 // derives from TRenderingCompositePassBase<InputCount, OutputCount>
@@ -26,7 +26,7 @@ class FRCPassPostProcessSelectionOutline : public TRenderingCompositePassBase<2,
 {
 public:
 	// interface FRenderingCompositePass ---------
-	virtual void Process(FRenderingCompositePassContext& Context) OVERRIDE;
-	virtual void Release() OVERRIDE { delete this; }
-	virtual FPooledRenderTargetDesc ComputeOutputDesc(EPassOutputId InPassOutputId) const OVERRIDE;
+	virtual void Process(FRenderingCompositePassContext& Context) override;
+	virtual void Release() override { delete this; }
+	virtual FPooledRenderTargetDesc ComputeOutputDesc(EPassOutputId InPassOutputId) const override;
 };
\ No newline at end of file
diff --git a/Engine/Source/Runtime/Renderer/Private/PostProcess/PostProcessTemporalAA.h b/Engine/Source/Runtime/Renderer/Private/PostProcess/PostProcessTemporalAA.h
--- a/Engine/Source/Runtime/Renderer/Private/PostProcess/PostProcessTemporalAA.h
+++ b/Engine/Source/Runtime/Renderer/Private/PostProcess/PostProcessTemporalAA.h
@@ -17,7 +17,7 @@ class FRCPassPostProcessSSRTemporalAA : public TRenderingCompositePassBase<3, 1>
 public:
 	// interface FRenderingCompositePass ---------
 	virtual void Process(FRenderingCompositePassContext& Context);
-	virtual void Release() OVERRIDE { delete this; }
+	virtual void Release() override { delete this; }
 	virtual FPooledRenderTargetDesc ComputeOutputDesc(EPassOutputId InPassOutputId) const;
 };
 
@@ -30,7 +30,7 @@ class FRCPassPostProcessDOFTemporalAA : public TRenderingCompositePassBase<3, 1>
 public:
 	// interface FRenderingCompositePass ---------
 	virtual void Process(FRenderingCompositePassContext& Context);
-	virtual void Release() OVERRIDE { delete this; }
+	virtual void Release() override { delete this; }
 	virtual FPooledRenderTargetDesc ComputeOutputDesc(EPassOutputId InPassOutputId) const;
 };
 
@@ -43,7 +43,7 @@ class FRCPassPostProcessLightShaftTemporalAA : public TRenderingCompositePassBas
 public:
 	// interface FRenderingCompositePass ---------
 	virtual void Process(FRenderingCompositePassContext& Context);
-	virtual void Release() OVERRIDE { delete this; }
+	virtual void Release() override { delete this; }
 	virtual FPooledRenderTargetDesc ComputeOutputDesc(EPassOutputId InPassOutputId) const;
 };
 
@@ -57,6 +57,6 @@ class FRCPassPostProcessTemporalAA : public TRenderingCompositePassBase<4, 1>
 public:
 	// interface FRenderingCompositePass ---------
 	virtual void Process(FRenderingCompositePassContext& Context);
-	virtual void Release() OVERRIDE { delete this; }
+	virtual void Release() override { delete this; }
 	virtual FPooledRenderTargetDesc ComputeOutputDesc(EPassOutputId InPassOutputId) const;
 };
\ No newline at end of file
diff --git a/Engine/Source/Runtime/Renderer/Private/PostProcess/SceneFilterRendering.cpp b/Engine/Source/Runtime/Renderer/Private/PostProcess/SceneFilterRendering.cpp
--- a/Engine/Source/Runtime/Renderer/Private/PostProcess/SceneFilterRendering.cpp
+++ b/Engine/Source/Runtime/Renderer/Private/PostProcess/SceneFilterRendering.cpp
@@ -16,7 +16,7 @@ class FScreenRectangleVertexBuffer : public FVertexBuffer
 {
 public:
 	/** Initialize the RHI for this rendering resource */
-	void InitRHI() OVERRIDE
+	void InitRHI() override
 	{
 		TResourceArray<FFilterVertex, VERTEXBUFFER_ALIGNMENT> Vertices;
 		Vertices.Init(6);
@@ -50,7 +50,7 @@ class FScreenRectangleIndexBuffer : public FIndexBuffer
 {
 public:
 	/** Initialize the RHI for this rendering resource */
-	void InitRHI() OVERRIDE
+	void InitRHI() override
 	{
 		// Indices 0 - 5 are used for rendering a quad. Indices 6 - 8 are used for triangle optimization. 
 		const uint16 Indices[] = { 0, 1, 2, 2, 1, 3, 0, 4, 5 };
diff --git a/Engine/Source/Runtime/Renderer/Private/RendererPrivate.h b/Engine/Source/Runtime/Renderer/Private/RendererPrivate.h
--- a/Engine/Source/Runtime/Renderer/Private/RendererPrivate.h
+++ b/Engine/Source/Runtime/Renderer/Private/RendererPrivate.h
@@ -19,26 +19,26 @@ DECLARE_LOG_CATEGORY_EXTERN(LogRenderer, Log, All);
 class FRendererModule : public IRendererModule
 {
 public:
-	virtual bool SupportsDynamicReloading() OVERRIDE { return true; }
-
-	virtual void BeginRenderingViewFamily(FCanvas* Canvas,const FSceneViewFamily* ViewFamily) OVERRIDE;
-	virtual FSceneInterface* AllocateScene(UWorld* World, bool bInRequiresHitProxies, ERHIFeatureLevel::Type InFeatureLevel) OVERRIDE;
-	virtual void RemoveScene(FSceneInterface* Scene) OVERRIDE;
-	virtual void UpdateStaticDrawListsForMaterials(const TArray<const FMaterial*>& Materials) OVERRIDE;
-	virtual FSceneViewStateInterface* AllocateViewState() OVERRIDE;
-	virtual uint32 GetNumDynamicLightsAffectingPrimitive(const FPrimitiveSceneInfo* PrimitiveSceneInfo,const FLightCacheInterface* LCI) OVERRIDE;
-	virtual void ReallocateSceneRenderTargets() OVERRIDE;
-	virtual void SceneRenderTargetsSetBufferSize(uint32 SizeX, uint32 SizeY) OVERRIDE;
-	virtual void DrawTileMesh(const FSceneView& View, const FMeshBatch& Mesh, bool bIsHitTesting, const FHitProxyId& HitProxyId) OVERRIDE;
-	virtual void RenderTargetPoolFindFreeElement(const FPooledRenderTargetDesc& Desc, TRefCountPtr<IPooledRenderTarget> &Out, const TCHAR* InDebugName) OVERRIDE;
-	virtual void TickRenderTargetPool() OVERRIDE;
-	virtual void DebugLogOnCrash() OVERRIDE;
-	virtual void GPUBenchmark(FSynthBenchmarkResults& InOut, uint32 WorkScale, bool bDebugOut) OVERRIDE;
-	virtual void QueryVisualizeTexture(FQueryVisualizeTexureInfo& Out) OVERRIDE;
-	virtual void ExecVisualizeTextureCmd(const FString& Cmd) OVERRIDE;
-	virtual void UpdateMapNeedsLightingFullyRebuiltState(UWorld* World) OVERRIDE;
-
-	virtual const TSet<FSceneInterface*>& GetAllocatedScenes() OVERRIDE
+	virtual bool SupportsDynamicReloading() override { return true; }
+
+	virtual void BeginRenderingViewFamily(FCanvas* Canvas,const FSceneViewFamily* ViewFamily) override;
+	virtual FSceneInterface* AllocateScene(UWorld* World, bool bInRequiresHitProxies, ERHIFeatureLevel::Type InFeatureLevel) override;
+	virtual void RemoveScene(FSceneInterface* Scene) override;
+	virtual void UpdateStaticDrawListsForMaterials(const TArray<const FMaterial*>& Materials) override;
+	virtual FSceneViewStateInterface* AllocateViewState() override;
+	virtual uint32 GetNumDynamicLightsAffectingPrimitive(const FPrimitiveSceneInfo* PrimitiveSceneInfo,const FLightCacheInterface* LCI) override;
+	virtual void ReallocateSceneRenderTargets() override;
+	virtual void SceneRenderTargetsSetBufferSize(uint32 SizeX, uint32 SizeY) override;
+	virtual void DrawTileMesh(const FSceneView& View, const FMeshBatch& Mesh, bool bIsHitTesting, const FHitProxyId& HitProxyId) override;
+	virtual void RenderTargetPoolFindFreeElement(const FPooledRenderTargetDesc& Desc, TRefCountPtr<IPooledRenderTarget> &Out, const TCHAR* InDebugName) override;
+	virtual void TickRenderTargetPool() override;
+	virtual void DebugLogOnCrash() override;
+	virtual void GPUBenchmark(FSynthBenchmarkResults& InOut, uint32 WorkScale, bool bDebugOut) override;
+	virtual void QueryVisualizeTexture(FQueryVisualizeTexureInfo& Out) override;
+	virtual void ExecVisualizeTextureCmd(const FString& Cmd) override;
+	virtual void UpdateMapNeedsLightingFullyRebuiltState(UWorld* World) override;
+
+	virtual const TSet<FSceneInterface*>& GetAllocatedScenes() override
 	{
 		return AllocatedScenes;
 	}
diff --git a/Engine/Source/Runtime/Renderer/Private/RendererScene.cpp b/Engine/Source/Runtime/Renderer/Private/RendererScene.cpp
--- a/Engine/Source/Runtime/Renderer/Private/RendererScene.cpp
+++ b/Engine/Source/Runtime/Renderer/Private/RendererScene.cpp
@@ -2014,7 +2014,7 @@ public:
 
 	virtual void AddDecal(UDecalComponent*){}
 	virtual void RemoveDecal(UDecalComponent*){}
-	virtual void UpdateDecalTransform(UDecalComponent* Decal) OVERRIDE {}
+	virtual void UpdateDecalTransform(UDecalComponent* Decal) override {}
 
 	/** Updates the transform of a light which has already been added to the scene. */
 	virtual void UpdateLightTransform(ULightComponent* Light){}
diff --git a/Engine/Source/Runtime/Renderer/Private/ScenePrivate.h b/Engine/Source/Runtime/Renderer/Private/ScenePrivate.h
--- a/Engine/Source/Runtime/Renderer/Private/ScenePrivate.h
+++ b/Engine/Source/Runtime/Renderer/Private/ScenePrivate.h
@@ -637,7 +637,7 @@ public:
 	// FSceneViewStateInterface
 	RENDERER_API virtual void Destroy();
 	
-	virtual void AddReferencedObjects(FReferenceCollector& Collector) OVERRIDE
+	virtual void AddReferencedObjects(FReferenceCollector& Collector) override
 	{
 		uint32 Count = MIDPool.Num();
 
@@ -650,7 +650,7 @@ public:
 	}
 
 	/** called in InitViews() */
-	virtual void OnStartFrame(FSceneView& CurrentView) OVERRIDE
+	virtual void OnStartFrame(FSceneView& CurrentView) override
 	{
 		check(IsInRenderingThread());
 
@@ -663,7 +663,7 @@ public:
 	}
 
 	// needed for GetReusableMID()
-	virtual void OnStartPostProcessing(FSceneView& CurrentView) OVERRIDE
+	virtual void OnStartPostProcessing(FSceneView& CurrentView) override
 	{
 		check(IsInGameThread());
 
@@ -707,7 +707,7 @@ public:
 		return Ret;
 	}
 
-	virtual void ApplyWorldOffset(FVector InOffset) OVERRIDE;
+	virtual void ApplyWorldOffset(FVector InOffset) override;
 
 	// FDeferredCleanupInterface
 	virtual void FinishCleanup()
@@ -1263,17 +1263,17 @@ public:
 	virtual void RemovePrimitive(UPrimitiveComponent* Primitive);
 	virtual void ReleasePrimitive(UPrimitiveComponent* Primitive);
 	virtual void UpdatePrimitiveTransform(UPrimitiveComponent* Primitive);
-	virtual void UpdatePrimitiveAttachment(UPrimitiveComponent* Primitive) OVERRIDE;
+	virtual void UpdatePrimitiveAttachment(UPrimitiveComponent* Primitive) override;
 	virtual void AddLight(ULightComponent* Light);
 	virtual void RemoveLight(ULightComponent* Light);
 	virtual void AddInvisibleLight(ULightComponent* Light);
 	virtual void SetSkyLight(FSkyLightSceneProxy* Light);
 	virtual void AddDecal(UDecalComponent* Component);
 	virtual void RemoveDecal(UDecalComponent* Component);
-	virtual void UpdateDecalTransform(UDecalComponent* Decal) OVERRIDE;
+	virtual void UpdateDecalTransform(UDecalComponent* Decal) override;
 	virtual void AddReflectionCapture(UReflectionCaptureComponent* Component);
 	virtual void RemoveReflectionCapture(UReflectionCaptureComponent* Component);
-	virtual void GetReflectionCaptureData(UReflectionCaptureComponent* Component, class FReflectionCaptureFullHDRDerivedData& OutDerivedData) OVERRIDE;
+	virtual void GetReflectionCaptureData(UReflectionCaptureComponent* Component, class FReflectionCaptureFullHDRDerivedData& OutDerivedData) override;
 	virtual void UpdateReflectionCaptureTransform(UReflectionCaptureComponent* Component);
 	virtual void ReleaseReflectionCubemap(UReflectionCaptureComponent* CaptureComponent);
 	virtual void UpdateSceneCaptureContents(class USceneCaptureComponent2D* CaptureComponent);
@@ -1299,9 +1299,9 @@ public:
 	virtual void UpdateSpeedTreeWind(double CurrentTime);
 	virtual FUniformBufferRHIParamRef GetSpeedTreeUniformBuffer(const FVertexFactory* VertexFactory);
 	virtual void DumpUnbuiltLightIteractions( FOutputDevice& Ar ) const;
-	virtual void DumpStaticMeshDrawListStats() const OVERRIDE;
+	virtual void DumpStaticMeshDrawListStats() const override;
 	virtual void SetClearMotionBlurInfoGameThread();
-	virtual void UpdateParameterCollections(const TArray<FMaterialParameterCollectionInstanceResource*>& InParameterCollections) OVERRIDE;
+	virtual void UpdateParameterCollections(const TArray<FMaterialParameterCollectionInstanceResource*>& InParameterCollections) override;
 
 	/** Determines whether the scene has dynamic sky lighting. */
 	bool HasDynamicSkyLighting() const
@@ -1393,16 +1393,16 @@ public:
 		return FUniformBufferRHIParamRef();
 	}
 
-	virtual void ApplyWorldOffset(FVector InOffset) OVERRIDE;
+	virtual void ApplyWorldOffset(FVector InOffset) override;
 
-	virtual void OnLevelAddedToWorld(FName InLevelName) OVERRIDE;
+	virtual void OnLevelAddedToWorld(FName InLevelName) override;
 
-	virtual bool HasAnyLights() const OVERRIDE { return NumVisibleLights > 0 || bHasSkyLight; }
+	virtual bool HasAnyLights() const override { return NumVisibleLights > 0 || bHasSkyLight; }
 
-	virtual bool IsEditorScene() const OVERRIDE { return bIsEditorScene; }
+	virtual bool IsEditorScene() const override { return bIsEditorScene; }
 
-	virtual ERHIFeatureLevel::Type GetFeatureLevel() const OVERRIDE { return FeatureLevel; }
-	virtual void ChangeFeatureLevel(ERHIFeatureLevel::Type InFeatureLevel) OVERRIDE;
+	virtual ERHIFeatureLevel::Type GetFeatureLevel() const override { return FeatureLevel; }
+	virtual void ChangeFeatureLevel(ERHIFeatureLevel::Type InFeatureLevel) override;
 
 private:
 
diff --git a/Engine/Source/Runtime/Renderer/Private/SceneRendering.h b/Engine/Source/Runtime/Renderer/Private/SceneRendering.h
--- a/Engine/Source/Runtime/Renderer/Private/SceneRendering.h
+++ b/Engine/Source/Runtime/Renderer/Private/SceneRendering.h
@@ -576,7 +576,7 @@ public:
 
 	// FSceneRenderer interface
 
-	virtual void Render() OVERRIDE;
+	virtual void Render() override;
 
 protected:
 
diff --git a/Engine/Source/Runtime/Renderer/Private/ShaderBaseClasses.h b/Engine/Source/Runtime/Renderer/Private/ShaderBaseClasses.h
--- a/Engine/Source/Runtime/Renderer/Private/ShaderBaseClasses.h
+++ b/Engine/Source/Runtime/Renderer/Private/ShaderBaseClasses.h
@@ -343,7 +343,7 @@ public:
 
 	// FShader interface.
 	virtual bool Serialize(FArchive& Ar);
-	virtual uint32 GetAllocatedSize() const OVERRIDE;
+	virtual uint32 GetAllocatedSize() const override;
 
 private:
 
@@ -435,7 +435,7 @@ public:
 	// FShader interface.
 	virtual const FVertexFactoryParameterRef* GetVertexFactoryParameterRef() const { return &VertexFactoryParameters; }
 	virtual bool Serialize(FArchive& Ar);
-	virtual uint32 GetAllocatedSize() const OVERRIDE;
+	virtual uint32 GetAllocatedSize() const override;
 
 private:
 	FVertexFactoryParameterRef VertexFactoryParameters;
diff --git a/Engine/Source/Runtime/Renderer/Private/ShadowRendering.h b/Engine/Source/Runtime/Renderer/Private/ShadowRendering.h
--- a/Engine/Source/Runtime/Renderer/Private/ShadowRendering.h
+++ b/Engine/Source/Runtime/Renderer/Private/ShadowRendering.h
@@ -39,7 +39,7 @@ namespace StencilingGeometry
 		/** 
 		* Initialize the RHI for this rendering resource 
 		*/
-		void InitRHI() OVERRIDE
+		void InitRHI() override
 		{
 			const int32 NumSides = NumSphereSides;
 			const int32 NumRings = NumSphereRings;
@@ -121,7 +121,7 @@ namespace StencilingGeometry
 		/** 
 		* Initialize the RHI for this rendering resource 
 		*/
-		void InitRHI() OVERRIDE
+		void InitRHI() override
 		{
 			const int32 NumSides = NumSphereSides;
 			const int32 NumRings = NumSphereRings;
@@ -169,7 +169,7 @@ namespace StencilingGeometry
 
 		static const uint32 NumVerts = NumSides * NumSlices * 2;
 
-		void InitRHI() OVERRIDE
+		void InitRHI() override
 		{
 			TResourceArray<uint16, INDEXBUFFER_ALIGNMENT> Indices;
 
@@ -241,7 +241,7 @@ namespace StencilingGeometry
 		/** 
 		* Initialize the RHI for this rendering resource 
 		*/
-		void InitRHI() OVERRIDE
+		void InitRHI() override
 		{
 			TResourceArray<FVector4, VERTEXBUFFER_ALIGNMENT> Verts;
 			Verts.Empty(NumVerts);
@@ -888,7 +888,7 @@ public:
 	void SetParameters(FRHICommandList& RHICmdList, const FSceneView& View, const FProjectedShadowInfo* ShadowInfo, int32 ShadowSplitIndex = INDEX_NONE);
 
 	// Begin FShader interface
-	virtual bool Serialize(FArchive& Ar) OVERRIDE
+	virtual bool Serialize(FArchive& Ar) override
 	{
 		bool bShaderHasOutdatedParameters = FGlobalShader::Serialize(Ar);
 		Ar << StencilingGeometryParameters;
@@ -1130,7 +1130,7 @@ public:
 		int32 ViewIndex,
 		const FSceneView& View,
 		const FProjectedShadowInfo* ShadowInfo
-		) OVERRIDE
+		) override
 	{
 		FShadowProjectionPixelShaderInterface::SetParameters(RHICmdList, ViewIndex,View,ShadowInfo);
 
@@ -1241,7 +1241,7 @@ public:
 		FRHICommandList& RHICmdList, 
 		int32 ViewIndex,
 		const FSceneView& View,
-		const FProjectedShadowInfo* ShadowInfo) OVERRIDE
+		const FProjectedShadowInfo* ShadowInfo) override
 	{
 		TShadowProjectionPS<Quality>::SetParameters(RHICmdList, ViewIndex, View, ShadowInfo);
 
diff --git a/Engine/Source/Runtime/SandboxFile/Private/IPlatformFileSandboxWrapper.cpp b/Engine/Source/Runtime/SandboxFile/Private/IPlatformFileSandboxWrapper.cpp
--- a/Engine/Source/Runtime/SandboxFile/Private/IPlatformFileSandboxWrapper.cpp
+++ b/Engine/Source/Runtime/SandboxFile/Private/IPlatformFileSandboxWrapper.cpp
@@ -290,7 +290,7 @@ FString FSandboxPlatformFile::ConvertToAbsolutePathForExternalAppForWrite( const
 class FSandboxFileModule : public IPlatformFileModule
 {
 public:
-	virtual IPlatformFile* GetPlatformFile() OVERRIDE
+	virtual IPlatformFile* GetPlatformFile() override
 	{
 		static TScopedPointer<IPlatformFile> AutoDestroySingleton(new FSandboxPlatformFile(true));
 		return AutoDestroySingleton.GetOwnedPointer();
diff --git a/Engine/Source/Runtime/SandboxFile/Public/IPlatformFileSandboxWrapper.h b/Engine/Source/Runtime/SandboxFile/Public/IPlatformFileSandboxWrapper.h
--- a/Engine/Source/Runtime/SandboxFile/Public/IPlatformFileSandboxWrapper.h
+++ b/Engine/Source/Runtime/SandboxFile/Public/IPlatformFileSandboxWrapper.h
@@ -26,27 +26,27 @@ public:
 	{
 	}
 
-	virtual int64		Tell() OVERRIDE
+	virtual int64		Tell() override
 	{
 		return FileHandle->Tell();
 	}
-	virtual bool		Seek(int64 NewPosition) OVERRIDE
+	virtual bool		Seek(int64 NewPosition) override
 	{
 		return FileHandle->Seek(NewPosition);
 	}
-	virtual bool		SeekFromEnd(int64 NewPositionRelativeToEnd) OVERRIDE
+	virtual bool		SeekFromEnd(int64 NewPositionRelativeToEnd) override
 	{
 		return FileHandle->SeekFromEnd(NewPositionRelativeToEnd);
 	}
-	virtual bool		Read(uint8* Destination, int64 BytesToRead) OVERRIDE
+	virtual bool		Read(uint8* Destination, int64 BytesToRead) override
 	{
 		return FileHandle->Read(Destination, BytesToRead);
 	}
-	virtual bool		Write(const uint8* Source, int64 BytesToWrite) OVERRIDE
+	virtual bool		Write(const uint8* Source, int64 BytesToWrite) override
 	{
 		return FileHandle->Write(Source, BytesToWrite);
 	}
-	virtual int64		Size() OVERRIDE
+	virtual int64		Size() override
 	{
 		return FileHandle->Size();
 	}
@@ -168,7 +168,7 @@ public:
 	 *
 	 *	@param	bInEnabled		true to enable the sandbox, false to disable it
 	 */
-	virtual void SetSandboxEnabled(bool bInEnabled) OVERRIDE
+	virtual void SetSandboxEnabled(bool bInEnabled) override
 	{
 		bSandboxEnabled = bInEnabled;
 	}
@@ -178,20 +178,20 @@ public:
 	 *
 	 *	@return	bool			true if enabled, false if not
 	 */
-	virtual bool IsSandboxEnabled() const OVERRIDE
+	virtual bool IsSandboxEnabled() const override
 	{
 		return bSandboxEnabled;
 	}
 
-	virtual bool ShouldBeUsed(IPlatformFile* Inner, const TCHAR* CmdLine) const OVERRIDE;
-	virtual bool Initialize(IPlatformFile* Inner, const TCHAR* CmdLine) OVERRIDE;
+	virtual bool ShouldBeUsed(IPlatformFile* Inner, const TCHAR* CmdLine) const override;
+	virtual bool Initialize(IPlatformFile* Inner, const TCHAR* CmdLine) override;
 
-	virtual IPlatformFile* GetLowerLevel() OVERRIDE
+	virtual IPlatformFile* GetLowerLevel() override
 	{
 		return LowerLevel;
 	}
 
-	virtual const TCHAR* GetName() const OVERRIDE
+	virtual const TCHAR* GetName() const override
 	{
 		return FSandboxPlatformFile::GetTypeName();
 	}
@@ -237,7 +237,7 @@ public:
 
 	// IPlatformFile Interface
 
-	virtual bool		FileExists(const TCHAR* Filename) OVERRIDE
+	virtual bool		FileExists(const TCHAR* Filename) override
 	{
 		// First look for the file in the user dir.
 		bool Result = LowerLevel->FileExists( *ConvertToSandboxPath( Filename ) );
@@ -248,7 +248,7 @@ public:
 		return Result;
 	}
 
-	virtual int64		FileSize(const TCHAR* Filename) OVERRIDE
+	virtual int64		FileSize(const TCHAR* Filename) override
 	{
 		// First look for the file in the user dir.
 		int64 Result = LowerLevel->FileSize( *ConvertToSandboxPath( Filename ) );
@@ -259,7 +259,7 @@ public:
 		return Result;
 	}
 
-	virtual bool		DeleteFile(const TCHAR* Filename) OVERRIDE
+	virtual bool		DeleteFile(const TCHAR* Filename) override
 	{
 		// Delete only sandbox files. If the sendbox version doesn't exists
 		// assume the delete was successful because we only care if the sandbox version is gone.
@@ -272,7 +272,7 @@ public:
 		return Result;
 	}
 
-	virtual bool		IsReadOnly(const TCHAR* Filename) OVERRIDE
+	virtual bool		IsReadOnly(const TCHAR* Filename) override
 	{
 		// If the file exists in the sandbox folder and is read-only return true
 		// Otherwise it can always be 'overwritten' in the sandbox
@@ -291,7 +291,7 @@ public:
 		return Result;
 	}
 
-	virtual bool		MoveFile(const TCHAR* To, const TCHAR* From) OVERRIDE
+	virtual bool		MoveFile(const TCHAR* To, const TCHAR* From) override
 	{
 		// Only files within the sandbox dir can be moved
 		bool Result = false;
@@ -303,7 +303,7 @@ public:
 		return Result;
 	}
 
-	virtual bool		SetReadOnly(const TCHAR* Filename, bool bNewReadOnlyValue) OVERRIDE
+	virtual bool		SetReadOnly(const TCHAR* Filename, bool bNewReadOnlyValue) override
 	{
 		bool Result = false;
 		FString UserFilename( *ConvertToSandboxPath( Filename ) );
@@ -314,7 +314,7 @@ public:
 		return Result;
 	}
 
-	virtual FDateTime		GetTimeStamp(const TCHAR* Filename) OVERRIDE
+	virtual FDateTime		GetTimeStamp(const TCHAR* Filename) override
 	{
 		FDateTime Result = FDateTime::MinValue();
 		FString UserFilename( *ConvertToSandboxPath( Filename ) );
@@ -329,7 +329,7 @@ public:
 		return Result;
 	}
 
-	virtual void 			SetTimeStamp(const TCHAR* Filename, FDateTime DateTime) OVERRIDE
+	virtual void 			SetTimeStamp(const TCHAR* Filename, FDateTime DateTime) override
 	{
 		FString UserFilename( *ConvertToSandboxPath( Filename ) );
 		if( LowerLevel->FileExists( *UserFilename ) )
@@ -342,7 +342,7 @@ public:
 		}
 	}
 
-	virtual FDateTime		GetAccessTimeStamp(const TCHAR* Filename) OVERRIDE
+	virtual FDateTime		GetAccessTimeStamp(const TCHAR* Filename) override
 	{
 		FDateTime Result = FDateTime::MinValue();
 		FString UserFilename( *ConvertToSandboxPath( Filename ) );
@@ -357,7 +357,7 @@ public:
 		return Result;
 	}
 
-	virtual IFileHandle*	OpenRead(const TCHAR* Filename) OVERRIDE
+	virtual IFileHandle*	OpenRead(const TCHAR* Filename) override
 	{
 		IFileHandle* Result = LowerLevel->OpenRead( *ConvertToSandboxPath( Filename ) );
 		if( !Result  && OkForInnerAccess(Filename) )
@@ -367,13 +367,13 @@ public:
 		return Result;
 	}
 
-	virtual IFileHandle*	OpenWrite(const TCHAR* Filename, bool bAppend = false, bool bAllowRead = false) OVERRIDE
+	virtual IFileHandle*	OpenWrite(const TCHAR* Filename, bool bAppend = false, bool bAllowRead = false) override
 	{
 		// Only files from the sandbox directory can be opened for wiriting
 		return LowerLevel->OpenWrite( *ConvertToSandboxPath( Filename ), bAppend, bAllowRead );
 	}
 
-	virtual bool		DirectoryExists(const TCHAR* Directory) OVERRIDE
+	virtual bool		DirectoryExists(const TCHAR* Directory) override
 	{
 		bool Result = LowerLevel->DirectoryExists( *ConvertToSandboxPath( Directory ) );
 		if( Result == false && OkForInnerAccess(Directory, true) )
@@ -383,13 +383,13 @@ public:
 		return Result;
 	}
 
-	virtual bool		CreateDirectory(const TCHAR* Directory) OVERRIDE
+	virtual bool		CreateDirectory(const TCHAR* Directory) override
 	{
 		// Directories can be created only under the sandbox path
 		return LowerLevel->CreateDirectory( *ConvertToSandboxPath( Directory ) );
 	}
 
-	virtual bool		DeleteDirectory(const TCHAR* Directory) OVERRIDE
+	virtual bool		DeleteDirectory(const TCHAR* Directory) override
 	{
 		// Directories can be deleted only under the sandbox path
 		return LowerLevel->DeleteDirectory( *ConvertToSandboxPath( Directory ) );
@@ -440,7 +440,7 @@ public:
 		}
 	};
 
-	virtual bool		IterateDirectory(const TCHAR* Directory, IPlatformFile::FDirectoryVisitor& Visitor) OVERRIDE
+	virtual bool		IterateDirectory(const TCHAR* Directory, IPlatformFile::FDirectoryVisitor& Visitor) override
 	{
 		FSandboxVisitor SandboxVisitor( Visitor, *this );
 		bool Result = false;
@@ -449,7 +449,7 @@ public:
 		return Result;
 	}
 
-	virtual bool		IterateDirectoryRecursively(const TCHAR* Directory, IPlatformFile::FDirectoryVisitor& Visitor) OVERRIDE
+	virtual bool		IterateDirectoryRecursively(const TCHAR* Directory, IPlatformFile::FDirectoryVisitor& Visitor) override
 	{
 		FSandboxVisitor SandboxVisitor( Visitor, *this );
 		bool Result = false;
@@ -458,19 +458,19 @@ public:
 		return Result;
 	}
 
-	virtual bool		DeleteDirectoryRecursively(const TCHAR* Directory) OVERRIDE
+	virtual bool		DeleteDirectoryRecursively(const TCHAR* Directory) override
 	{
 		// Directories can be deleted only under the sandbox path
 		return LowerLevel->DeleteDirectoryRecursively( *ConvertToSandboxPath( Directory ) );
 	}
 
-	virtual bool CreateDirectoryTree(const TCHAR* Directory) OVERRIDE
+	virtual bool CreateDirectoryTree(const TCHAR* Directory) override
 	{
 		// Directories can only be created only under the sandbox path
 		return LowerLevel->CreateDirectoryTree( *ConvertToSandboxPath( Directory ) );
 	}
 
-	virtual bool		CopyFile(const TCHAR* To, const TCHAR* From) OVERRIDE
+	virtual bool		CopyFile(const TCHAR* To, const TCHAR* From) override
 	{
 		// Files can be copied only to the sandbox directory
 		bool Result = false;
@@ -485,8 +485,8 @@ public:
 		return Result;
 	}
 
-	virtual FString ConvertToAbsolutePathForExternalAppForRead( const TCHAR* Filename ) OVERRIDE;
-	virtual FString ConvertToAbsolutePathForExternalAppForWrite( const TCHAR* Filename ) OVERRIDE;
+	virtual FString ConvertToAbsolutePathForExternalAppForRead( const TCHAR* Filename ) override;
+	virtual FString ConvertToAbsolutePathForExternalAppForWrite( const TCHAR* Filename ) override;
 };
 
 
diff --git a/Engine/Source/Runtime/ShaderCore/Public/Shader.h b/Engine/Source/Runtime/ShaderCore/Public/Shader.h
--- a/Engine/Source/Runtime/ShaderCore/Public/Shader.h
+++ b/Engine/Source/Runtime/ShaderCore/Public/Shader.h
@@ -824,7 +824,7 @@ private:
 	static FShader* ConstructSerializedInstance() { return new ShaderClass(); } \
 	static FShader* ConstructCompiledInstance(const ShaderMetaType::CompiledShaderInitializerType& Initializer) \
 	{ return new ShaderClass(Initializer); } \
-	virtual uint32 GetTypeSize() const OVERRIDE { return sizeof(*this); }
+	virtual uint32 GetTypeSize() const override { return sizeof(*this); }
 #define DECLARE_SHADER_TYPE(ShaderClass,ShaderMetaTypeShortcut) \
 	DECLARE_EXPORTED_SHADER_TYPE(ShaderClass,ShaderMetaTypeShortcut,)
 
diff --git a/Engine/Source/Runtime/Sockets/Private/Android/SocketSubsystemAndroid.h b/Engine/Source/Runtime/Sockets/Private/Android/SocketSubsystemAndroid.h
--- a/Engine/Source/Runtime/Sockets/Private/Android/SocketSubsystemAndroid.h
+++ b/Engine/Source/Runtime/Sockets/Private/Android/SocketSubsystemAndroid.h
@@ -52,15 +52,15 @@ public:
 	 *
 	 * @return TRUE if initialized ok, FALSE otherwise
 	 */
-	virtual bool Init(FString& Error) OVERRIDE;
+	virtual bool Init(FString& Error) override;
 
 	/**
 	 * Performs platform specific socket clean up
 	 */
-	virtual void Shutdown() OVERRIDE;
+	virtual void Shutdown() override;
 
 	/**
 	 * @return Whether the device has a properly configured network device or not
 	 */
-	virtual bool HasNetworkDevice() OVERRIDE;
+	virtual bool HasNetworkDevice() override;
 };
diff --git a/Engine/Source/Runtime/Sockets/Private/BSDIPv6Sockets/IPAddressBSDIPv6.h b/Engine/Source/Runtime/Sockets/Private/BSDIPv6Sockets/IPAddressBSDIPv6.h
--- a/Engine/Source/Runtime/Sockets/Private/BSDIPv6Sockets/IPAddressBSDIPv6.h
+++ b/Engine/Source/Runtime/Sockets/Private/BSDIPv6Sockets/IPAddressBSDIPv6.h
@@ -33,7 +33,7 @@ public:
 	 *
 	 * @param InAddr the new address to use (must convert to network byte order)
 	 */
-	virtual void SetIp(uint32 InAddr) OVERRIDE
+	virtual void SetIp(uint32 InAddr) override
 	{
 		if(InAddr == 0)
 		{
@@ -61,7 +61,7 @@ public:
 	 * @param InAddr the string containing the new ip address to use
 	 * @param bIsValid will be set to true if InAddr was a valid IPv6 or IPv4 address, false if not.
 	 */
-	virtual void SetIp(const TCHAR* InAddr, bool& bIsValid) OVERRIDE
+	virtual void SetIp(const TCHAR* InAddr, bool& bIsValid) override
 	{
 		// check for valid IPv6 address
 		auto InAddrAnsi = StringCast<ANSICHAR>(InAddr);
@@ -159,7 +159,7 @@ public:
 	 *
 	 * @param OutAddr the out param receiving the ip address
 	 */
-	virtual void GetIp(uint32& OutAddr) const OVERRIDE
+	virtual void GetIp(uint32& OutAddr) const override
 	{
 		// grab the last 32-bits of the IPv6 address as this will correspond to the IPv4 address
 		// in a dual stack system.
@@ -186,7 +186,7 @@ public:
 	 *
 	 * @param InPort the new port to use (must convert to network byte order)
 	 */
-	virtual void SetPort(int32 InPort) OVERRIDE
+	virtual void SetPort(int32 InPort) override
 	{
 		Addr.sin6_port = htons(InPort);
 	}
@@ -196,26 +196,26 @@ public:
 	 *
 	 * @param OutPort the host byte order int that receives the port
 	 */
-	virtual void GetPort(int32& OutPort) const OVERRIDE
+	virtual void GetPort(int32& OutPort) const override
 	{
 		OutPort = ntohs(Addr.sin6_port);
 	}
 
 	/** Returns the port number from this address in host byte order */
-	virtual int32 GetPort() const OVERRIDE
+	virtual int32 GetPort() const override
 	{
 		return ntohs(Addr.sin6_port);
 	}
 
 	/** Sets the address to be any address */
-	virtual void SetAnyAddress() OVERRIDE
+	virtual void SetAnyAddress() override
 	{
 		SetIp(in6addr_any);
 		SetPort(0);
 	}
 
 	/** Sets the address to broadcast */
-	virtual void SetBroadcastAddress() OVERRIDE
+	virtual void SetBroadcastAddress() override
 	{
 		// broadcast means something different in IPv6, but this is a rough equivalent
 #ifndef in6addr_allnodesonlink
@@ -235,7 +235,7 @@ public:
 	 *
 	 * @param bAppendPort whether to append the port information or not
 	 */
-	virtual FString ToString(bool bAppendPort) const OVERRIDE
+	virtual FString ToString(bool bAppendPort) const override
 	{
 		char IPStr[INET6_ADDRSTRLEN];
 
@@ -259,7 +259,7 @@ public:
 	 *
 	 * @param Other the address to compare against
 	 */
-	virtual bool operator==(const FInternetAddr& Other) const OVERRIDE
+	virtual bool operator==(const FInternetAddr& Other) const override
 	{
 		FInternetAddrBSDIPv6& OtherBSD = (FInternetAddrBSDIPv6&)Other;
 		return memcmp(&Addr.sin6_addr,&OtherBSD.Addr.sin6_addr,sizeof(in6_addr)) == 0 &&
@@ -272,7 +272,7 @@ public:
 	 *
 	 * @return true if a valid IP, false otherwise
 	 */
-	virtual bool IsValid() const OVERRIDE
+	virtual bool IsValid() const override
 	{
 		FInternetAddrBSDIPv6 Temp;
 
diff --git a/Engine/Source/Runtime/Sockets/Private/BSDIPv6Sockets/SocketSubsystemBSDIPv6.h b/Engine/Source/Runtime/Sockets/Private/BSDIPv6Sockets/SocketSubsystemBSDIPv6.h
--- a/Engine/Source/Runtime/Sockets/Private/BSDIPv6Sockets/SocketSubsystemBSDIPv6.h
+++ b/Engine/Source/Runtime/Sockets/Private/BSDIPv6Sockets/SocketSubsystemBSDIPv6.h
@@ -21,19 +21,19 @@ public:
 
 	// Begin ISocketSubsystem interface
 
-	virtual TSharedRef<FInternetAddr> CreateInternetAddr( uint32 Address = 0, uint32 Port = 0 ) OVERRIDE;
+	virtual TSharedRef<FInternetAddr> CreateInternetAddr( uint32 Address = 0, uint32 Port = 0 ) override;
 
-	virtual class FSocket* CreateSocket( const FName& SocketType, const FString& SocketDescription, bool bForceUDP = false ) OVERRIDE;
+	virtual class FSocket* CreateSocket( const FName& SocketType, const FString& SocketDescription, bool bForceUDP = false ) override;
 
-	virtual void DestroySocket( class FSocket* Socket ) OVERRIDE;
+	virtual void DestroySocket( class FSocket* Socket ) override;
 
-	virtual ESocketErrors GetHostByName( const ANSICHAR* HostName, FInternetAddr& OutAddr ) OVERRIDE;
+	virtual ESocketErrors GetHostByName( const ANSICHAR* HostName, FInternetAddr& OutAddr ) override;
 
-	virtual bool GetHostName( FString& HostName ) OVERRIDE;
+	virtual bool GetHostName( FString& HostName ) override;
 
-	virtual ESocketErrors GetLastErrorCode( ) OVERRIDE;
+	virtual ESocketErrors GetLastErrorCode( ) override;
 
-	virtual bool GetLocalAdapterAddresses( TArray<TSharedPtr<FInternetAddr> >& OutAdresses ) OVERRIDE
+	virtual bool GetLocalAdapterAddresses( TArray<TSharedPtr<FInternetAddr> >& OutAdresses ) override
 	{
 		bool bCanBindAll;
 
@@ -42,19 +42,19 @@ public:
 		return true;
 	}
 
-	virtual const TCHAR* GetSocketAPIName( ) const OVERRIDE;
+	virtual const TCHAR* GetSocketAPIName( ) const override;
 
-	virtual bool RequiresChatDataBeSeparate( ) OVERRIDE
+	virtual bool RequiresChatDataBeSeparate( ) override
 	{
 		return false;
 	}
 
-	virtual bool RequiresEncryptedPackets( ) OVERRIDE
+	virtual bool RequiresEncryptedPackets( ) override
 	{
 		return false;
 	}
 
-	virtual ESocketErrors TranslateErrorCode( int32 Code ) OVERRIDE;
+	virtual ESocketErrors TranslateErrorCode( int32 Code ) override;
 
 	// End ISocketSubsystem interface
 
diff --git a/Engine/Source/Runtime/Sockets/Private/BSDIPv6Sockets/SocketsBSDIPv6.h b/Engine/Source/Runtime/Sockets/Private/BSDIPv6Sockets/SocketsBSDIPv6.h
--- a/Engine/Source/Runtime/Sockets/Private/BSDIPv6Sockets/SocketsBSDIPv6.h
+++ b/Engine/Source/Runtime/Sockets/Private/BSDIPv6Sockets/SocketsBSDIPv6.h
@@ -60,59 +60,59 @@ public:
 
 public:
 
-	virtual bool Close() OVERRIDE;
+	virtual bool Close() override;
 
-	virtual bool Bind(const FInternetAddr& Addr) OVERRIDE;
+	virtual bool Bind(const FInternetAddr& Addr) override;
 
-	virtual bool Connect(const FInternetAddr& Addr) OVERRIDE;
+	virtual bool Connect(const FInternetAddr& Addr) override;
 
-	virtual bool Listen(int32 MaxBacklog) OVERRIDE;
+	virtual bool Listen(int32 MaxBacklog) override;
 
-	virtual bool HasPendingConnection(bool& bHasPendingConnection) OVERRIDE;
+	virtual bool HasPendingConnection(bool& bHasPendingConnection) override;
 
-	virtual bool HasPendingData(uint32& PendingDataSize) OVERRIDE;
+	virtual bool HasPendingData(uint32& PendingDataSize) override;
 
-	virtual class FSocket* Accept(const FString& SocketDescription) OVERRIDE;
+	virtual class FSocket* Accept(const FString& SocketDescription) override;
 
-	virtual class FSocket* Accept(FInternetAddr& OutAddr, const FString& SocketDescription) OVERRIDE;
+	virtual class FSocket* Accept(FInternetAddr& OutAddr, const FString& SocketDescription) override;
 
-	virtual bool SendTo(const uint8* Data, int32 Count, int32& BytesSent, const FInternetAddr& Destination) OVERRIDE;
+	virtual bool SendTo(const uint8* Data, int32 Count, int32& BytesSent, const FInternetAddr& Destination) override;
 
-	virtual bool Send(const uint8* Data, int32 Count, int32& BytesSent) OVERRIDE;
+	virtual bool Send(const uint8* Data, int32 Count, int32& BytesSent) override;
 
-	virtual bool RecvFrom(uint8* Data, int32 BufferSize, int32& BytesRead, FInternetAddr& Source, ESocketReceiveFlags::Type Flags = ESocketReceiveFlags::None) OVERRIDE;
+	virtual bool RecvFrom(uint8* Data, int32 BufferSize, int32& BytesRead, FInternetAddr& Source, ESocketReceiveFlags::Type Flags = ESocketReceiveFlags::None) override;
 
-	virtual bool Recv(uint8* Data,int32 BufferSize,int32& BytesRead, ESocketReceiveFlags::Type Flags = ESocketReceiveFlags::None) OVERRIDE;
+	virtual bool Recv(uint8* Data,int32 BufferSize,int32& BytesRead, ESocketReceiveFlags::Type Flags = ESocketReceiveFlags::None) override;
 
-	virtual bool Wait(ESocketWaitConditions::Type Condition, FTimespan WaitTime) OVERRIDE;
+	virtual bool Wait(ESocketWaitConditions::Type Condition, FTimespan WaitTime) override;
 
-	virtual ESocketConnectionState GetConnectionState() OVERRIDE;
+	virtual ESocketConnectionState GetConnectionState() override;
 
-	virtual void GetAddress(FInternetAddr& OutAddr) OVERRIDE;
+	virtual void GetAddress(FInternetAddr& OutAddr) override;
 
-	virtual bool SetNonBlocking(bool bIsNonBlocking = true) OVERRIDE;
+	virtual bool SetNonBlocking(bool bIsNonBlocking = true) override;
 
-	virtual bool SetBroadcast(bool bAllowBroadcast = true) OVERRIDE;
+	virtual bool SetBroadcast(bool bAllowBroadcast = true) override;
 
-	virtual bool JoinMulticastGroup(const FInternetAddr& GroupAddress) OVERRIDE;
+	virtual bool JoinMulticastGroup(const FInternetAddr& GroupAddress) override;
 
-	virtual bool LeaveMulticastGroup(const FInternetAddr& GroupAddress) OVERRIDE;
+	virtual bool LeaveMulticastGroup(const FInternetAddr& GroupAddress) override;
 
-	virtual bool SetMulticastLoopback(bool bLoopback) OVERRIDE;
+	virtual bool SetMulticastLoopback(bool bLoopback) override;
 
-	virtual bool SetMulticastTtl(uint8 TimeToLive) OVERRIDE;
+	virtual bool SetMulticastTtl(uint8 TimeToLive) override;
 
-	virtual bool SetReuseAddr(bool bAllowReuse = true) OVERRIDE;
+	virtual bool SetReuseAddr(bool bAllowReuse = true) override;
 
-	virtual bool SetLinger(bool bShouldLinger = true, int32 Timeout = 0) OVERRIDE;
+	virtual bool SetLinger(bool bShouldLinger = true, int32 Timeout = 0) override;
 
-	virtual bool SetRecvErr(bool bUseErrorQueue = true) OVERRIDE;
+	virtual bool SetRecvErr(bool bUseErrorQueue = true) override;
 
-	virtual bool SetSendBufferSize(int32 Size,int32& NewSize) OVERRIDE;
+	virtual bool SetSendBufferSize(int32 Size,int32& NewSize) override;
 
-	virtual bool SetReceiveBufferSize(int32 Size,int32& NewSize) OVERRIDE;
+	virtual bool SetReceiveBufferSize(int32 Size,int32& NewSize) override;
 
-	virtual int32 GetPortNo() OVERRIDE;
+	virtual int32 GetPortNo() override;
 
 	bool SetIPv6Only(bool bIPv6Only);
 
diff --git a/Engine/Source/Runtime/Sockets/Private/BSDSockets/IPAddressBSD.h b/Engine/Source/Runtime/Sockets/Private/BSDSockets/IPAddressBSD.h
--- a/Engine/Source/Runtime/Sockets/Private/BSDSockets/IPAddressBSD.h
+++ b/Engine/Source/Runtime/Sockets/Private/BSDSockets/IPAddressBSD.h
@@ -29,7 +29,7 @@ public:
 	 *
 	 * @param InAddr the new address to use (must convert to network byte order)
 	 */
-	virtual void SetIp(uint32 InAddr) OVERRIDE
+	virtual void SetIp(uint32 InAddr) override
 	{
 		Addr.sin_addr.s_addr = htonl(InAddr);
 	}
@@ -39,7 +39,7 @@ public:
 	 *
 	 * @param InAddr the string containing the new ip address to use
 	 */
-	virtual void SetIp(const TCHAR* InAddr, bool& bIsValid) OVERRIDE
+	virtual void SetIp(const TCHAR* InAddr, bool& bIsValid) override
 	{
 		int32 A, B, C, D;
 		int32 Port = 0;
@@ -105,7 +105,7 @@ public:
 	 *
 	 * @param OutAddr the out param receiving the ip address
 	 */
-	virtual void GetIp(uint32& OutAddr) const OVERRIDE
+	virtual void GetIp(uint32& OutAddr) const override
 	{ 
 		OutAddr = ntohl(Addr.sin_addr.s_addr);
 	}
@@ -125,7 +125,7 @@ public:
 	 *
 	 * @param InPort the new port to use (must convert to network byte order)
 	 */
-	virtual void SetPort(int32 InPort) OVERRIDE
+	virtual void SetPort(int32 InPort) override
 	{
 		Addr.sin_port = htons(InPort);
 	}
@@ -135,7 +135,7 @@ public:
 	 *
 	 * @param OutPort the host byte order int that receives the port
 	 */
-	virtual void GetPort(int32& OutPort) const OVERRIDE
+	virtual void GetPort(int32& OutPort) const override
 	{
 		OutPort = ntohs(Addr.sin_port);
 	}
@@ -143,20 +143,20 @@ public:
 	/**
 	 * Returns the port number from this address in host byte order
 	 */
-	virtual int32 GetPort(void) const OVERRIDE
+	virtual int32 GetPort(void) const override
 	{
 		return ntohs(Addr.sin_port);
 	}
 
 	/** Sets the address to be any address */
-	virtual void SetAnyAddress(void) OVERRIDE
+	virtual void SetAnyAddress(void) override
 	{
 		SetIp(INADDR_ANY);
 		SetPort(0);
 	}
 
 	/** Sets the address to broadcast */
-	virtual void SetBroadcastAddress() OVERRIDE
+	virtual void SetBroadcastAddress() override
 	{
 		SetIp(INADDR_BROADCAST);
 		SetPort(0);
@@ -167,7 +167,7 @@ public:
 	 *
 	 * @param bAppendPort whether to append the port information or not
 	 */
-	virtual FString ToString(bool bAppendPort) const OVERRIDE
+	virtual FString ToString(bool bAppendPort) const override
 	{
 		uint32 LocalAddr = ntohl(Addr.sin_addr.s_addr);
 
@@ -191,7 +191,7 @@ public:
 	 *
 	 * @param Other the address to compare against
 	 */
-	virtual bool operator==(const FInternetAddr& Other) const OVERRIDE
+	virtual bool operator==(const FInternetAddr& Other) const override
 	{
 		FInternetAddrBSD& OtherBSD = (FInternetAddrBSD&)Other;
 		return Addr.sin_addr.s_addr == OtherBSD.Addr.sin_addr.s_addr &&
@@ -204,7 +204,7 @@ public:
 	 *
 	 * @return true if a valid IP, false otherwise
 	 */
-	virtual bool IsValid() const OVERRIDE
+	virtual bool IsValid() const override
 	{
 		return Addr.sin_addr.s_addr != 0;
 	}
diff --git a/Engine/Source/Runtime/Sockets/Private/BSDSockets/SocketSubsystemBSD.h b/Engine/Source/Runtime/Sockets/Private/BSDSockets/SocketSubsystemBSD.h
--- a/Engine/Source/Runtime/Sockets/Private/BSDSockets/SocketSubsystemBSD.h
+++ b/Engine/Source/Runtime/Sockets/Private/BSDSockets/SocketSubsystemBSD.h
@@ -21,19 +21,19 @@ public:
 
 	// Begin ISocketSubsystem interface
 
-	virtual TSharedRef<FInternetAddr> CreateInternetAddr( uint32 Address = 0, uint32 Port = 0 ) OVERRIDE;
+	virtual TSharedRef<FInternetAddr> CreateInternetAddr( uint32 Address = 0, uint32 Port = 0 ) override;
 
-	virtual class FSocket* CreateSocket( const FName& SocketType, const FString& SocketDescription, bool bForceUDP = false ) OVERRIDE;
+	virtual class FSocket* CreateSocket( const FName& SocketType, const FString& SocketDescription, bool bForceUDP = false ) override;
 
-	virtual void DestroySocket( class FSocket* Socket ) OVERRIDE;
+	virtual void DestroySocket( class FSocket* Socket ) override;
 
-	virtual ESocketErrors GetHostByName( const ANSICHAR* HostName, FInternetAddr& OutAddr ) OVERRIDE;
+	virtual ESocketErrors GetHostByName( const ANSICHAR* HostName, FInternetAddr& OutAddr ) override;
 
-	virtual bool GetHostName( FString& HostName ) OVERRIDE;
+	virtual bool GetHostName( FString& HostName ) override;
 
-	virtual ESocketErrors GetLastErrorCode( ) OVERRIDE;
+	virtual ESocketErrors GetLastErrorCode( ) override;
 
-	virtual bool GetLocalAdapterAddresses( TArray<TSharedPtr<FInternetAddr> >& OutAdresses ) OVERRIDE
+	virtual bool GetLocalAdapterAddresses( TArray<TSharedPtr<FInternetAddr> >& OutAdresses ) override
 	{
 		bool bCanBindAll;
 
@@ -42,19 +42,19 @@ public:
 		return true;
 	}
 
-	virtual const TCHAR* GetSocketAPIName( ) const OVERRIDE;
+	virtual const TCHAR* GetSocketAPIName( ) const override;
 
-	virtual bool RequiresChatDataBeSeparate( ) OVERRIDE
+	virtual bool RequiresChatDataBeSeparate( ) override
 	{
 		return false;
 	}
 
-	virtual bool RequiresEncryptedPackets( ) OVERRIDE
+	virtual bool RequiresEncryptedPackets( ) override
 	{
 		return false;
 	}
 
-	virtual ESocketErrors TranslateErrorCode( int32 Code ) OVERRIDE;
+	virtual ESocketErrors TranslateErrorCode( int32 Code ) override;
 
 	// End ISocketSubsystem interface
 
diff --git a/Engine/Source/Runtime/Sockets/Private/BSDSockets/SocketsBSD.h b/Engine/Source/Runtime/Sockets/Private/BSDSockets/SocketsBSD.h
--- a/Engine/Source/Runtime/Sockets/Private/BSDSockets/SocketsBSD.h
+++ b/Engine/Source/Runtime/Sockets/Private/BSDSockets/SocketsBSD.h
@@ -63,59 +63,59 @@ public:
 
 public:
 
-	virtual bool Close() OVERRIDE;
+	virtual bool Close() override;
 
-	virtual bool Bind(const FInternetAddr& Addr) OVERRIDE;
+	virtual bool Bind(const FInternetAddr& Addr) override;
 
-	virtual bool Connect(const FInternetAddr& Addr) OVERRIDE;
+	virtual bool Connect(const FInternetAddr& Addr) override;
 
-	virtual bool Listen(int32 MaxBacklog) OVERRIDE;
+	virtual bool Listen(int32 MaxBacklog) override;
 
-	virtual bool HasPendingConnection(bool& bHasPendingConnection) OVERRIDE;
+	virtual bool HasPendingConnection(bool& bHasPendingConnection) override;
 
-	virtual bool HasPendingData(uint32& PendingDataSize) OVERRIDE;
+	virtual bool HasPendingData(uint32& PendingDataSize) override;
 
-	virtual class FSocket* Accept(const FString& SocketDescription) OVERRIDE;
+	virtual class FSocket* Accept(const FString& SocketDescription) override;
 
-	virtual class FSocket* Accept(FInternetAddr& OutAddr, const FString& SocketDescription) OVERRIDE;
+	virtual class FSocket* Accept(FInternetAddr& OutAddr, const FString& SocketDescription) override;
 
-	virtual bool SendTo(const uint8* Data, int32 Count, int32& BytesSent, const FInternetAddr& Destination) OVERRIDE;
+	virtual bool SendTo(const uint8* Data, int32 Count, int32& BytesSent, const FInternetAddr& Destination) override;
 
-	virtual bool Send(const uint8* Data, int32 Count, int32& BytesSent) OVERRIDE;
+	virtual bool Send(const uint8* Data, int32 Count, int32& BytesSent) override;
 
-	virtual bool RecvFrom(uint8* Data, int32 BufferSize, int32& BytesRead, FInternetAddr& Source, ESocketReceiveFlags::Type Flags = ESocketReceiveFlags::None) OVERRIDE;
+	virtual bool RecvFrom(uint8* Data, int32 BufferSize, int32& BytesRead, FInternetAddr& Source, ESocketReceiveFlags::Type Flags = ESocketReceiveFlags::None) override;
 
-	virtual bool Recv(uint8* Data,int32 BufferSize,int32& BytesRead, ESocketReceiveFlags::Type Flags = ESocketReceiveFlags::None) OVERRIDE;
+	virtual bool Recv(uint8* Data,int32 BufferSize,int32& BytesRead, ESocketReceiveFlags::Type Flags = ESocketReceiveFlags::None) override;
 
-	virtual bool Wait(ESocketWaitConditions::Type Condition, FTimespan WaitTime) OVERRIDE;
+	virtual bool Wait(ESocketWaitConditions::Type Condition, FTimespan WaitTime) override;
 
-	virtual ESocketConnectionState GetConnectionState() OVERRIDE;
+	virtual ESocketConnectionState GetConnectionState() override;
 
-	virtual void GetAddress(FInternetAddr& OutAddr) OVERRIDE;
+	virtual void GetAddress(FInternetAddr& OutAddr) override;
 
-	virtual bool SetNonBlocking(bool bIsNonBlocking = true) OVERRIDE;
+	virtual bool SetNonBlocking(bool bIsNonBlocking = true) override;
 
-	virtual bool SetBroadcast(bool bAllowBroadcast = true) OVERRIDE;
+	virtual bool SetBroadcast(bool bAllowBroadcast = true) override;
 
-	virtual bool JoinMulticastGroup(const FInternetAddr& GroupAddress) OVERRIDE;
+	virtual bool JoinMulticastGroup(const FInternetAddr& GroupAddress) override;
 
-	virtual bool LeaveMulticastGroup(const FInternetAddr& GroupAddress) OVERRIDE;
+	virtual bool LeaveMulticastGroup(const FInternetAddr& GroupAddress) override;
 
-	virtual bool SetMulticastLoopback(bool bLoopback) OVERRIDE;
+	virtual bool SetMulticastLoopback(bool bLoopback) override;
 
-	virtual bool SetMulticastTtl(uint8 TimeToLive) OVERRIDE;
+	virtual bool SetMulticastTtl(uint8 TimeToLive) override;
 
-	virtual bool SetReuseAddr(bool bAllowReuse = true) OVERRIDE;
+	virtual bool SetReuseAddr(bool bAllowReuse = true) override;
 
-	virtual bool SetLinger(bool bShouldLinger = true, int32 Timeout = 0) OVERRIDE;
+	virtual bool SetLinger(bool bShouldLinger = true, int32 Timeout = 0) override;
 
-	virtual bool SetRecvErr(bool bUseErrorQueue = true) OVERRIDE;
+	virtual bool SetRecvErr(bool bUseErrorQueue = true) override;
 
-	virtual bool SetSendBufferSize(int32 Size,int32& NewSize) OVERRIDE;
+	virtual bool SetSendBufferSize(int32 Size,int32& NewSize) override;
 
-	virtual bool SetReceiveBufferSize(int32 Size,int32& NewSize) OVERRIDE;
+	virtual bool SetReceiveBufferSize(int32 Size,int32& NewSize) override;
 
-	virtual int32 GetPortNo() OVERRIDE;
+	virtual int32 GetPortNo() override;
 
 	/**
 	 * Return the Socket for anyone who knows they have an FSocketBSD
diff --git a/Engine/Source/Runtime/Sockets/Private/HTML5/Device/SocketSubsystem.h b/Engine/Source/Runtime/Sockets/Private/HTML5/Device/SocketSubsystem.h
--- a/Engine/Source/Runtime/Sockets/Private/HTML5/Device/SocketSubsystem.h
+++ b/Engine/Source/Runtime/Sockets/Private/HTML5/Device/SocketSubsystem.h
@@ -49,15 +49,15 @@ public:
 	 *
 	 * @return TRUE if initialized ok, FALSE otherwise
 	 */
-	virtual bool Init(FString& Error) OVERRIDE;
+	virtual bool Init(FString& Error) override;
 
 	/**
 	 * Performs platform specific socket clean up
 	 */
-	virtual void Shutdown() OVERRIDE;
+	virtual void Shutdown() override;
 
 	/**
 	 * @return Whether the device has a properly configured network device or not
 	 */
-	virtual bool HasNetworkDevice() OVERRIDE;
+	virtual bool HasNetworkDevice() override;
 };
diff --git a/Engine/Source/Runtime/Sockets/Private/HTML5/Simulator/IPAddressX.h b/Engine/Source/Runtime/Sockets/Private/HTML5/Simulator/IPAddressX.h
--- a/Engine/Source/Runtime/Sockets/Private/HTML5/Simulator/IPAddressX.h
+++ b/Engine/Source/Runtime/Sockets/Private/HTML5/Simulator/IPAddressX.h
@@ -24,58 +24,58 @@ public:
 	 *
 	 * @param InAddr the new address to use (must convert to network byte order)
 	 */
-	virtual void SetIp(uint32 InAddr) OVERRIDE;
+	virtual void SetIp(uint32 InAddr) override;
 	/**
 	 * Sets the ip address from a string ("A.B.C.D")
 	 *
 	 * @param InAddr the string containing the new ip address to use
 	 */
-	virtual void SetIp(const TCHAR* InAddr, bool& bIsValid) OVERRIDE;
+	virtual void SetIp(const TCHAR* InAddr, bool& bIsValid) override;
 
 	/**
 	 * Copies the network byte order ip address to a host byte order dword
 	 *
 	 * @param OutAddr the out param receiving the ip address
 	 */
-	virtual void GetIp(uint32& OutAddr) const OVERRIDE;
+	virtual void GetIp(uint32& OutAddr) const override;
 
 	/**
 	 * Sets the port number from a host byte order int
 	 *
 	 * @param InPort the new port to use (must convert to network byte order)
 	 */
-	virtual void SetPort(int32 InPort) OVERRIDE;
+	virtual void SetPort(int32 InPort) override;
 
 	/**
 	 * Copies the port number from this address and places it into a host byte order int
 	 *
 	 * @param OutPort the host byte order int that receives the port
 	 */
-	virtual void GetPort(int32& OutPort) const OVERRIDE;
+	virtual void GetPort(int32& OutPort) const override;
 
 	/**
 	 * Returns the port number from this address in host byte order
 	 */
-	virtual int32 GetPort(void) const OVERRIDE;
+	virtual int32 GetPort(void) const override;
 
 	/** Sets the address to be any address */
-	virtual void SetAnyAddress(void) OVERRIDE;
+	virtual void SetAnyAddress(void) override;
 
 	/** Sets the address to broadcast */
-	virtual void SetBroadcastAddress() OVERRIDE;
+	virtual void SetBroadcastAddress() override;
 
 	/**
 	 * Converts this internet ip address to string form
 	 *
 	 * @param bAppendPort whether to append the port information or not
 	 */
-	virtual FString ToString(bool bAppendPort) const OVERRIDE;
+	virtual FString ToString(bool bAppendPort) const override;
 	/**
 	 * Is this a well formed internet address
 	 *
 	 * @return true if a valid IP, false otherwise
 	 */
-	virtual bool IsValid() const OVERRIDE;
+	virtual bool IsValid() const override;
 
 	FInternetAddrRaw* GetPimpl() { return Pimpl; }
 
