Fixing FSandboxVisitor not working properly with game projects not directly under UE4 root folder.

diff --git a/Engine/Source/Runtime/SandboxFile/Private/IPlatformFileSandboxWrapper.cpp b/Engine/Source/Runtime/SandboxFile/Private/IPlatformFileSandboxWrapper.cpp
--- a/Engine/Source/Runtime/SandboxFile/Private/IPlatformFileSandboxWrapper.cpp
+++ b/Engine/Source/Runtime/SandboxFile/Private/IPlatformFileSandboxWrapper.cpp
@@ -284,6 +284,20 @@ FString FSandboxPlatformFile::ConvertToAbsolutePathForExternalAppForWrite( const
 	return ConvertToSandboxPath( Filename );
 }
 
+const FString& FSandboxPlatformFile::GetAbsolutePathToGameDirectory()
+{
+	if (AbsoluteGameDirectory.IsEmpty())
+	{
+		AbsoluteGameDirectory = FPaths::GetProjectFilePath();
+		UE_CLOG(AbsoluteGameDirectory.IsEmpty(), SandboxFile, Fatal, TEXT("SandboxFileWrapper tried to access project path before it was set."));
+		AbsoluteGameDirectory = FPaths::ConvertRelativePathToFull(AbsoluteGameDirectory);
+		// Strip .uproject filename and game directory, keep just to path to the game directory which could simply be the root dir (but not always).
+		AbsoluteGameDirectory = FPaths::GetPath(FPaths::GetPath(AbsoluteGameDirectory));
+	}
+	return AbsoluteGameDirectory;
+}
+
+
 /**
  * Module for the sandbox file
  */
diff --git a/Engine/Source/Runtime/SandboxFile/Public/IPlatformFileSandboxWrapper.h b/Engine/Source/Runtime/SandboxFile/Public/IPlatformFileSandboxWrapper.h
--- a/Engine/Source/Runtime/SandboxFile/Public/IPlatformFileSandboxWrapper.h
+++ b/Engine/Source/Runtime/SandboxFile/Public/IPlatformFileSandboxWrapper.h
@@ -62,6 +62,8 @@ class SANDBOXFILE_API FSandboxPlatformFile : public IPlatformFile
 	FString			RelativeRootDirectory;
 	/** Absolute path to root directory. Cached for faster access */
 	FString			AbsoluteRootDirectory;
+	/** Absolute path to game directory. Cached for faster access */
+	FString			AbsoluteGameDirectory;
 	/** Access to any file (in unreal standard form) matching this is not allowed */
 	TArray<FString>		FileExclusionWildcards;
 	/** Access to any directory (in unreal standard form) matching this is not allowed */
@@ -211,12 +213,15 @@ public:
 		return SandboxDirectory;
 	}
 
-	/** Returns absolute game directory */
+	/** Returns absolute root directory */
 	const FString& GetAbsoluteRootDirectory() const
 	{
 		return AbsoluteRootDirectory;
 	}
 
+	/** Returns absolute path to game directory (without the game directory itself) */
+	const FString& GetAbsolutePathToGameDirectory();
+
 	/** 
 	 * Add exclusion. These files and / or directories pretend not to exist so that they cannot be accessed at all (except in the sandbox) 
 	 * @param Wildcard FString::MatchesWildcard-type wild card to test for exclusion
@@ -434,7 +439,15 @@ public:
 				VisitedSandboxFiles.Add( *LocalFilename );
 				// Now convert the sandbox path back to engine path because the sandbox folder should not be exposed
 				// to the engine and remain transparent.
-				LocalFilename = LocalFilename.Replace( *SandboxFile.GetSandboxDirectory(), *SandboxFile.GetAbsoluteRootDirectory(), ESearchCase::IgnoreCase );
+				LocalFilename = LocalFilename.Mid(SandboxFile.GetSandboxDirectory().Len());
+				if (LocalFilename.StartsWith(TEXT("Engine/")))
+				{
+					LocalFilename = SandboxFile.GetAbsoluteRootDirectory() / LocalFilename;
+				}
+				else
+				{
+					LocalFilename = SandboxFile.GetAbsolutePathToGameDirectory() / LocalFilename;
+				}
 			}
 			else
 			{
