Core - More stats for memory profiling

diff --git a/Engine/Source/Runtime/CoreUObject/Private/Misc/PackageName.cpp b/Engine/Source/Runtime/CoreUObject/Private/Misc/PackageName.cpp
--- a/Engine/Source/Runtime/CoreUObject/Private/Misc/PackageName.cpp
+++ b/Engine/Source/Runtime/CoreUObject/Private/Misc/PackageName.cpp
@@ -658,6 +658,8 @@ bool FPackageName::DoesPackageExist(const FString& LongPackageName, const FGuid*
 
 bool FPackageName::SearchForPackageOnDisk(const FString& PackageName, FString* OutLongPackageName, FString* OutFilename, bool bUseLocalizedNames)
 {
+	DECLARE_SCOPE_CYCLE_COUNTER(TEXT("FPackageName::SearchForPackageOnDisk"), STAT_PackageName_SearchForPackageOnDisk, STATGROUP_LoadTime);
+
 	bool bResult = false;
 	double StartTime = FPlatformTime::Seconds();
 	if (FPackageName::IsShortPackageName(PackageName) == false)
diff --git a/Engine/Source/Runtime/SandboxFile/Private/IPlatformFileSandboxWrapper.cpp b/Engine/Source/Runtime/SandboxFile/Private/IPlatformFileSandboxWrapper.cpp
--- a/Engine/Source/Runtime/SandboxFile/Private/IPlatformFileSandboxWrapper.cpp
+++ b/Engine/Source/Runtime/SandboxFile/Private/IPlatformFileSandboxWrapper.cpp
@@ -131,6 +131,9 @@ bool FSandboxPlatformFile::Initialize(IPlatformFile* Inner, const TCHAR* CmdLine
 
 FString FSandboxPlatformFile::ConvertToSandboxPath( const TCHAR* Filename ) const
 {
+	// Mostly for the malloc profiler to flush the data.
+	DECLARE_SCOPE_CYCLE_COUNTER(TEXT("FSandboxPlatformFile::ConvertToSandboxPath"), STAT_SandboxPlatformFile_ConvertToSandboxPath, STATGROUP_LoadTimeVerbose);
+
 	// convert to a standardized path (relative)
 	FString SandboxPath = Filename;
 	FPaths::MakeStandardFilename(SandboxPath);
