PR #6212: SandboxDirectory should be absolute and have no relative paths in it (Contributed by [REDACTED_PERSON]).
The actual change got lost when submitting 8672165.

diff --git a/Engine/Source/Runtime/SandboxFile/Private/IPlatformFileSandboxWrapper.cpp b/Engine/Source/Runtime/SandboxFile/Private/IPlatformFileSandboxWrapper.cpp
--- a/Engine/Source/Runtime/SandboxFile/Private/IPlatformFileSandboxWrapper.cpp
+++ b/Engine/Source/Runtime/SandboxFile/Private/IPlatformFileSandboxWrapper.cpp
@@ -118,7 +118,7 @@ bool FSandboxPlatformFile::Initialize(IPlatformFile* Inner, const TCHAR* CmdLine
 			FPaths::MakeStandardFilename(SandboxDirectory);
 
 			// SandboxDirectory should be absolute and have no relative paths in it
-			FPaths::ConvertRelativePathToFull(SandboxDirectory);
+			SandboxDirectory = FPaths::ConvertRelativePathToFull(SandboxDirectory);
 		}
 
 		if( bWipeSandbox )
