Fix for incorrect sandboxed paths

When the argument is an already sandboxed path ConvertToSandboxPath added another copy of the path to the end. We need to test if the path is inside the full sandboxed path before testing if it's inside the project directory, since the former is inside the latter.

diff --git a/Engine/Source/Runtime/SandboxFile/Private/IPlatformFileSandboxWrapper.cpp b/Engine/Source/Runtime/SandboxFile/Private/IPlatformFileSandboxWrapper.cpp
--- a/Engine/Source/Runtime/SandboxFile/Private/IPlatformFileSandboxWrapper.cpp
+++ b/Engine/Source/Runtime/SandboxFile/Private/IPlatformFileSandboxWrapper.cpp
@@ -162,18 +162,24 @@ FString FSandboxPlatformFile::ConvertToSandboxPath( const TCHAR* Filename ) cons
 		// See whether Filename is relative to root directory.
 		// if it's not inside the root, then just use it
 		FString FullSandboxPath = FPaths::ConvertRelativePathToFull(SandboxPath);
-		FString FullGameDir;
+		FString FullGameDir, FullSandboxedGameDir;
 #if IS_PROGRAM
 		if (FPaths::IsProjectFilePathSet())
 		{
 			FullGameDir = FPaths::ConvertRelativePathToFull(FPaths::GetPath(FPaths::GetProjectFilePath()) + TEXT("/"));
+			FullSandboxedGameDir = FPaths::Combine(*SandboxDirectory, *FPaths::GetBaseFilename(FPaths::GetProjectFilePath()));
 		}
 		else
 #endif
 		{
 			FullGameDir = FPaths::ConvertRelativePathToFull(FPaths::ProjectDir());
+			FullSandboxedGameDir = FPaths::Combine(*SandboxDirectory, FApp::GetProjectName());
 		}
-		if(FullSandboxPath.StartsWith(FullGameDir))
+		if(FullSandboxPath.StartsWith(FullSandboxedGameDir))
+		{
+			return SandboxPath;
+		}
+		else if (FullSandboxPath.StartsWith(FullGameDir))
 		{
 #if IS_PROGRAM
 			SandboxPath = FPaths::Combine(*SandboxDirectory, *FPaths::GetBaseFilename(FPaths::GetProjectFilePath()), *FullSandboxPath + FullGameDir.Len());
