default sandbox usage can be overridden per-platform.

diff --git a/Engine/Source/Runtime/SandboxFile/Private/IPlatformFileSandboxWrapper.cpp b/Engine/Source/Runtime/SandboxFile/Private/IPlatformFileSandboxWrapper.cpp
--- a/Engine/Source/Runtime/SandboxFile/Private/IPlatformFileSandboxWrapper.cpp
+++ b/Engine/Source/Runtime/SandboxFile/Private/IPlatformFileSandboxWrapper.cpp
@@ -12,6 +12,10 @@
 
 DEFINE_LOG_CATEGORY(SandboxFile);
 
+#if !defined(PLATFORM_SUPPORTS_DEFAULT_SANDBOX)
+	#define PLATFORM_SUPPORTS_DEFAULT_SANDBOX PLATFORM_DESKTOP
+#endif
+
 FSandboxPlatformFile::FSandboxPlatformFile(bool bInEntireEngineWillUseThisSandbox)
 	: LowerLevel(NULL)
 	, bEntireEngineWillUseThisSandbox(bInEntireEngineWillUseThisSandbox)
@@ -28,7 +32,7 @@ bool FSandboxPlatformFile::ShouldBeUsed(IPlatformFile* Inner, const TCHAR* CmdLi
 {
 	FString SandboxDir;
 	bool bResult = FParse::Value( CmdLine, TEXT("-Sandbox="), SandboxDir );
-#if PLATFORM_DESKTOP && (UE_GAME || UE_SERVER)
+#if PLATFORM_SUPPORTS_DEFAULT_SANDBOX && (UE_GAME || UE_SERVER)
 	if (FPlatformProperties::RequiresCookedData() && SandboxDir.IsEmpty() && Inner == &FPlatformFileManager::Get().GetPlatformFile() && bEntireEngineWillUseThisSandbox)
 	{
 		SandboxDir = GetCookedSandboxDir();
@@ -42,7 +46,7 @@ bool FSandboxPlatformFile::Initialize(IPlatformFile* Inner, const TCHAR* CmdLine
 {
 	FString CommandLineDirectory;
 	FParse::Value( CmdLine, TEXT("-Sandbox="), CommandLineDirectory);
-#if PLATFORM_DESKTOP && (UE_GAME || UE_SERVER)
+#if PLATFORM_SUPPORTS_DEFAULT_SANDBOX && (UE_GAME || UE_SERVER)
 	if (CommandLineDirectory.IsEmpty() && bEntireEngineWillUseThisSandbox)
 	{
 		CommandLineDirectory = GetCookedSandboxDir();
