Removed references to UE4

diff --git a/Engine/Source/Runtime/SandboxFile/Private/IPlatformFileSandboxWrapper.cpp b/Engine/Source/Runtime/SandboxFile/Private/IPlatformFileSandboxWrapper.cpp
--- a/Engine/Source/Runtime/SandboxFile/Private/IPlatformFileSandboxWrapper.cpp
+++ b/Engine/Source/Runtime/SandboxFile/Private/IPlatformFileSandboxWrapper.cpp
@@ -128,7 +128,7 @@ bool FSandboxPlatformFile::Initialize(IPlatformFile* Inner, const TCHAR* CmdLine
 			// Special case - platform defined user directory will be used
 			SandboxDirectory = FPlatformProcess::UserDir();
 			SandboxDirectory += TEXT("My Games/");
-			SandboxDirectory += TEXT( "UE4/" );
+			SandboxDirectory += TEXT( "UE/" );
 			bSandboxIsAbsolute = true;
 		}
 		else if( CommandLineDirectory == TEXT("Unique") )
