Fixed incorrect GetCookedSandboxDir() path on Mac.

FPath::ProjectSavedDir() can be influenced by command line arguments such as `-userdir=` or the installed status of the engine and thus can return directories that are not where cooked output is expected to be. The `-installed` command line argument is added to UE_GAME-type targets on Mac, such that ProjectSavedDir() does not return the same path the cooker used (which uses ProjectDir() + "/Saved"). As a result, Mac does not work the same way other desktop platforms work during development.

#rnx

diff --git a/Engine/Source/Runtime/SandboxFile/Private/IPlatformFileSandboxWrapper.cpp b/Engine/Source/Runtime/SandboxFile/Private/IPlatformFileSandboxWrapper.cpp
--- a/Engine/Source/Runtime/SandboxFile/Private/IPlatformFileSandboxWrapper.cpp
+++ b/Engine/Source/Runtime/SandboxFile/Private/IPlatformFileSandboxWrapper.cpp
@@ -19,7 +19,12 @@ DEFINE_LOG_CATEGORY(SandboxFile);
 #if PLATFORM_SUPPORTS_DEFAULT_SANDBOX && (UE_GAME || UE_SERVER)
 static FString GetCookedSandboxDir()
 {
-	return FPaths::Combine(*(FPaths::ProjectSavedDir()), TEXT("Cooked"), ANSI_TO_TCHAR(FPlatformProperties::PlatformName()));
+	TStringBuilder<512> Dir;
+	Dir << *FPaths::ProjectDir();
+	Dir << "Saved/Cooked/";
+	Dir << FPlatformProperties::PlatformName();
+	Dir << "/";
+	return Dir.ToString();
 }
 #endif
 
