Fix bug with sandbox file wrapper where the Iterate* functions could return a value that was incorrect if the visitor had returned false during the iteration process.

diff --git a/Engine/Source/Runtime/SandboxFile/Private/IPlatformFileSandboxWrapper.cpp b/Engine/Source/Runtime/SandboxFile/Private/IPlatformFileSandboxWrapper.cpp
--- a/Engine/Source/Runtime/SandboxFile/Private/IPlatformFileSandboxWrapper.cpp
+++ b/Engine/Source/Runtime/SandboxFile/Private/IPlatformFileSandboxWrapper.cpp
@@ -807,8 +807,8 @@ bool FSandboxPlatformFile::IterateDirectory(const TCHAR* Directory, IPlatformFil
 	FSandboxVisitor SandboxVisitor(Visitor, *this, false);
 	FString SandboxDir = ConvertToSandboxPath(Directory);
 	bool Result = LowerLevel->IterateDirectory(*SandboxDir, SandboxVisitor);
-	// don't iterate the same directory twice if the Convert didn't change the path
-	if (SandboxDir != Directory)
+	// don't iterate the same directory twice if the Convert didn't change the path, or if we were asked to stop iterating
+	if (Result && (SandboxDir != Directory))
 	{
 		Result = LowerLevel->IterateDirectory(Directory, SandboxVisitor);
 	}
@@ -820,8 +820,8 @@ bool FSandboxPlatformFile::IterateDirectoryRecursively(const TCHAR* Directory, I
 	FSandboxVisitor SandboxVisitor(Visitor, *this, true);
 	FString SandboxDir = ConvertToSandboxPath(Directory);
 	bool Result = LowerLevel->IterateDirectoryRecursively(*SandboxDir, SandboxVisitor);
-	// don't iterate the same directory twice if the Convert didn't change the path
-	if (SandboxDir != Directory)
+	// don't iterate the same directory twice if the Convert didn't change the path, or if we were asked to stop iterating
+	if (Result && (SandboxDir != Directory))
 	{
 		Result = LowerLevel->IterateDirectoryRecursively(Directory, SandboxVisitor);
 	}
@@ -922,8 +922,8 @@ bool		FSandboxPlatformFile::IterateDirectoryStat(const TCHAR* Directory, IPlatfo
 	FSandboxStatVisitor SandboxVisitor(Visitor, *this, false);
 	FString SandboxDir = ConvertToSandboxPath(Directory);
 	bool Result = LowerLevel->IterateDirectoryStat(*SandboxDir, SandboxVisitor);
-	// don't iterate the same directory twice if the Convert didn't change the path (which can happen with injection type)
-	if (SandboxDir != Directory)
+	// don't iterate the same directory twice if the Convert didn't change the path, or if we were asked to stop iterating (which can happen with injection type)
+	if (Result && (SandboxDir != Directory))
 	{
 		Result = LowerLevel->IterateDirectoryStat(Directory, SandboxVisitor);
 	}
@@ -935,8 +935,8 @@ bool		FSandboxPlatformFile::IterateDirectoryStatRecursively(const TCHAR* Directo
 	FSandboxStatVisitor SandboxVisitor(Visitor, *this, true);
 	FString SandboxDir = ConvertToSandboxPath(Directory);
 	bool Result = LowerLevel->IterateDirectoryStatRecursively(*ConvertToSandboxPath(Directory), SandboxVisitor);
-	// don't iterate the same directory twice if the Convert didn't change the path (which can happen with injection type)
-	if (SandboxDir != Directory)
+	// don't iterate the same directory twice if the Convert didn't change the path, or if we were asked to stop iterating (which can happen with injection type)
+	if (Result && (SandboxDir != Directory))
 	{
 		Result = LowerLevel->IterateDirectoryStatRecursively(Directory, SandboxVisitor);
 	}
