[Engine]

diff --git a/Engine/Source/Runtime/RHICore/Public/RHIDescriptorAllocator.h b/Engine/Source/Runtime/RHICore/Public/RHIDescriptorAllocator.h
--- a/Engine/Source/Runtime/RHICore/Public/RHIDescriptorAllocator.h
+++ b/Engine/Source/Runtime/RHICore/Public/RHIDescriptorAllocator.h
@@ -8,21 +8,21 @@
 
 struct FRHIDescriptorAllocatorRange;
 
-class RHICORE_API FRHIDescriptorAllocator
+class FRHIDescriptorAllocator
 {
 public:
-	FRHIDescriptorAllocator();
-	FRHIDescriptorAllocator(uint32 InNumDescriptors, TConstArrayView<TStatId> InStats);
-	~FRHIDescriptorAllocator();
+	RHICORE_API FRHIDescriptorAllocator();
+	RHICORE_API FRHIDescriptorAllocator(uint32 InNumDescriptors, TConstArrayView<TStatId> InStats);
+	RHICORE_API ~FRHIDescriptorAllocator();
 
-	void Init(uint32 InNumDescriptors, TConstArrayView<TStatId> InStats);
-	void Shutdown();
+	RHICORE_API void Init(uint32 InNumDescriptors, TConstArrayView<TStatId> InStats);
+	RHICORE_API void Shutdown();
 
-	FRHIDescriptorHandle Allocate(ERHIDescriptorHeapType InType);
-	void Free(FRHIDescriptorHandle InHandle);
+	RHICORE_API FRHIDescriptorHandle Allocate(ERHIDescriptorHeapType InType);
+	RHICORE_API void Free(FRHIDescriptorHandle InHandle);
 
-	bool Allocate(uint32 NumDescriptors, uint32& OutSlot);
-	void Free(uint32 Slot, uint32 NumDescriptors);
+	RHICORE_API bool Allocate(uint32 NumDescriptors, uint32& OutSlot);
+	RHICORE_API void Free(uint32 Slot, uint32 NumDescriptors);
 
 	inline uint32 GetCapacity() const { return Capacity; }
 
@@ -57,17 +57,17 @@ private:
 #endif
 };
 
-class RHICORE_API FRHIHeapDescriptorAllocator : protected FRHIDescriptorAllocator
+class FRHIHeapDescriptorAllocator : protected FRHIDescriptorAllocator
 {
 public:
 	FRHIHeapDescriptorAllocator() = delete;
-	FRHIHeapDescriptorAllocator(ERHIDescriptorHeapType InType, uint32 InDescriptorCount, TConstArrayView<TStatId> InStats);
+	RHICORE_API FRHIHeapDescriptorAllocator(ERHIDescriptorHeapType InType, uint32 InDescriptorCount, TConstArrayView<TStatId> InStats);
 
-	FRHIDescriptorHandle Allocate();
-	void Free(FRHIDescriptorHandle InHandle);
+	RHICORE_API FRHIDescriptorHandle Allocate();
+	RHICORE_API void Free(FRHIDescriptorHandle InHandle);
 
-	bool Allocate(uint32 NumDescriptors, uint32& OutSlot);
-	void Free(uint32 Slot, uint32 NumDescriptors);
+	RHICORE_API bool Allocate(uint32 NumDescriptors, uint32& OutSlot);
+	RHICORE_API void Free(uint32 Slot, uint32 NumDescriptors);
 
 	using FRHIDescriptorAllocator::GetCapacity;
 	inline ERHIDescriptorHeapType GetType() const { return Type; }
@@ -78,14 +78,14 @@ private:
 	ERHIDescriptorHeapType Type;
 };
 
-class RHICORE_API FRHIOffsetHeapDescriptorAllocator : protected FRHIHeapDescriptorAllocator
+class FRHIOffsetHeapDescriptorAllocator : protected FRHIHeapDescriptorAllocator
 {
 public:
 	FRHIOffsetHeapDescriptorAllocator() = delete;
-	FRHIOffsetHeapDescriptorAllocator(ERHIDescriptorHeapType InType, uint32 InDescriptorCount, uint32 InHeapOffset, TConstArrayView<TStatId> InStats);
+	RHICORE_API FRHIOffsetHeapDescriptorAllocator(ERHIDescriptorHeapType InType, uint32 InDescriptorCount, uint32 InHeapOffset, TConstArrayView<TStatId> InStats);
 
-	FRHIDescriptorHandle Allocate();
-	void Free(FRHIDescriptorHandle InHandle);
+	RHICORE_API FRHIDescriptorHandle Allocate();
+	RHICORE_API void Free(FRHIDescriptorHandle InHandle);
 
 	using FRHIHeapDescriptorAllocator::GetCapacity;
 	using FRHIHeapDescriptorAllocator::GetType;
diff --git a/Engine/Source/Runtime/RHICore/Public/RHIPoolAllocator.h b/Engine/Source/Runtime/RHICore/Public/RHIPoolAllocator.h
--- a/Engine/Source/Runtime/RHICore/Public/RHIPoolAllocator.h
+++ b/Engine/Source/Runtime/RHICore/Public/RHIPoolAllocator.h
@@ -30,25 +30,25 @@ enum class ERHIPoolResourceTypes
 /**
 @brief Pool allocator internal data
 **/
-struct RHICORE_API FRHIPoolAllocationData
+struct FRHIPoolAllocationData
 {
 public:
 
 	// Reset to unknown state
-	void Reset();
+	RHICORE_API void Reset();
 
 	// Initialize options
-	void InitAsHead(int16 InPoolIndex);
-	void InitAsFree(int16 InPoolIndex, uint32 InSize, uint32 InAlignment, uint32 InOffset);
-	void InitAsAllocated(uint32 InSize, uint32 InPoolAlignment, uint32 InAllocationAlignment, FRHIPoolAllocationData* InFree);
-	void MoveFrom(FRHIPoolAllocationData& InAllocated, bool InLocked);
+	RHICORE_API void InitAsHead(int16 InPoolIndex);
+	RHICORE_API void InitAsFree(int16 InPoolIndex, uint32 InSize, uint32 InAlignment, uint32 InOffset);
+	RHICORE_API void InitAsAllocated(uint32 InSize, uint32 InPoolAlignment, uint32 InAllocationAlignment, FRHIPoolAllocationData* InFree);
+	RHICORE_API void MoveFrom(FRHIPoolAllocationData& InAllocated, bool InLocked);
 
 	// Free block operation (TODO: make internal)
-	void MarkFree(uint32 InPoolAlignment, uint32 InAllocationAlignment);
+	RHICORE_API void MarkFree(uint32 InPoolAlignment, uint32 InAllocationAlignment);
 
 	// Alias operations (NOTE: currently not accessed when pool lock is taken, needs to be fixed)
-	void AddAlias(FRHIPoolAllocationData* InOther);
-	void RemoveAlias();
+	RHICORE_API void AddAlias(FRHIPoolAllocationData* InOther);
+	RHICORE_API void RemoveAlias();
 	FRHIPoolAllocationData* GetFirstAlias() const { return AliasAllocation; }
 
 	// Type getters
@@ -81,12 +81,12 @@ private:
 	friend class FRHIPoolAllocator;
 	void Lock() { check(Locked == 0); Locked = 1; }
 
-	void Merge(FRHIPoolAllocationData* InOther);
+	RHICORE_API void Merge(FRHIPoolAllocationData* InOther);
 
 	// Linked list operation
-	void RemoveFromLinkedList();
-	void AddBefore(FRHIPoolAllocationData* InOther);
-	void AddAfter(FRHIPoolAllocationData* InOther);
+	RHICORE_API void RemoveFromLinkedList();
+	RHICORE_API void AddBefore(FRHIPoolAllocationData* InOther);
+	RHICORE_API void AddAfter(FRHIPoolAllocationData* InOther);
 		
 	enum class EAllocationType : uint8
 	{
@@ -126,7 +126,7 @@ private:
 @brief Pool which stores for each allocation the previous and next allocation.
 	   Each block is either free or allocated.
 **/
-class RHICORE_API FRHIMemoryPool
+class FRHIMemoryPool
 {
 public:
 
@@ -137,19 +137,19 @@ public:
 	};
 
 	// Constructor
-	FRHIMemoryPool(int16 InPoolIndex, uint64 InPoolSize, uint32 InPoolAlignment, ERHIPoolResourceTypes InSupportedResourceTypes, EFreeListOrder InFreeListOrder);
-	virtual ~FRHIMemoryPool();
+	RHICORE_API FRHIMemoryPool(int16 InPoolIndex, uint64 InPoolSize, uint32 InPoolAlignment, ERHIPoolResourceTypes InSupportedResourceTypes, EFreeListOrder InFreeListOrder);
+	RHICORE_API virtual ~FRHIMemoryPool();
 
 	// Setup/Shutdown
-	virtual void Init();
-	virtual void Destroy();
+	RHICORE_API virtual void Init();
+	RHICORE_API virtual void Destroy();
 
 	// AllocationData functions: allocate, free, move, ...
-	bool TryAllocate(uint32 InSizeInBytes, uint32 InAllocationAlignment, ERHIPoolResourceTypes InAllocationResourceType, FRHIPoolAllocationData& AllocationData);
-	void Deallocate(FRHIPoolAllocationData& AllocationData);
+	RHICORE_API bool TryAllocate(uint32 InSizeInBytes, uint32 InAllocationAlignment, ERHIPoolResourceTypes InAllocationResourceType, FRHIPoolAllocationData& AllocationData);
+	RHICORE_API void Deallocate(FRHIPoolAllocationData& AllocationData);
 
 	// Bookkeeping and clearing
-	void TryClear(FRHIPoolAllocator* InAllocator, uint32 InMaxCopySize, uint32& CopySize, const TArray<FRHIMemoryPool*>& InTargetPools);
+	RHICORE_API void TryClear(FRHIPoolAllocator* InAllocator, uint32 InMaxCopySize, uint32& CopySize, const TArray<FRHIMemoryPool*>& InTargetPools);
 
 	// Getters
 	int16 GetPoolIndex() const { return PoolIndex; }
@@ -164,22 +164,22 @@ public:
 	bool IsFull() const { return FreeSize == 0; }
 
 	// Static public helper functions
-	static uint32 GetAlignedSize(uint32 InSizeInBytes, uint32 InPoolAlignment, uint32 InAllocationAlignment);
-	static uint32 GetAlignedOffset(uint32 InOffset, uint32 InPoolAlignment, uint32 InAllocationAlignment);
+	static RHICORE_API uint32 GetAlignedSize(uint32 InSizeInBytes, uint32 InPoolAlignment, uint32 InAllocationAlignment);
+	static RHICORE_API uint32 GetAlignedOffset(uint32 InOffset, uint32 InPoolAlignment, uint32 InAllocationAlignment);
 
 protected:
 		
 	// Free block management helper function
-	int32 FindFreeBlock(uint32 InSizeInBytes, uint32 InAllocationAlignment) const;
-	FRHIPoolAllocationData* AddToFreeBlocks(FRHIPoolAllocationData* InFreeBlock);
-	void RemoveFromFreeBlocks(FRHIPoolAllocationData* InFreeBlock);
+	RHICORE_API int32 FindFreeBlock(uint32 InSizeInBytes, uint32 InAllocationAlignment) const;
+	RHICORE_API FRHIPoolAllocationData* AddToFreeBlocks(FRHIPoolAllocationData* InFreeBlock);
+	RHICORE_API void RemoveFromFreeBlocks(FRHIPoolAllocationData* InFreeBlock);
 
 	// Get/Release allocation data blocks
-	FRHIPoolAllocationData* GetNewAllocationData();
-	void ReleaseAllocationData(FRHIPoolAllocationData* InData);
+	RHICORE_API FRHIPoolAllocationData* GetNewAllocationData();
+	RHICORE_API void ReleaseAllocationData(FRHIPoolAllocationData* InData);
 
 	// Validate the linked list data - slow
-	void Validate();
+	RHICORE_API void Validate();
 
 	// Const creation members
 	int16 PoolIndex;
@@ -211,29 +211,29 @@ protected:
 /**
 @brief Manages an array of FRHIMemoryPool and supports defragmentation between multiple pools
 **/
-class RHICORE_API FRHIPoolAllocator
+class FRHIPoolAllocator
 {
 public:
 
 	// Constructor
-	FRHIPoolAllocator(uint64 InDefaultPoolSize, uint32 InPoolAlignment, uint32 InMaxAllocationSize, FRHIMemoryPool::EFreeListOrder InFreeListOrder, bool InDefragEnabled);
-	virtual ~FRHIPoolAllocator();
+	RHICORE_API FRHIPoolAllocator(uint64 InDefaultPoolSize, uint32 InPoolAlignment, uint32 InMaxAllocationSize, FRHIMemoryPool::EFreeListOrder InFreeListOrder, bool InDefragEnabled);
+	RHICORE_API virtual ~FRHIPoolAllocator();
 
 	// Setup/Shutdown
-	void Initialize();
-	void Destroy();
+	RHICORE_API void Initialize();
+	RHICORE_API void Destroy();
 
 	// Defrag & cleanup operation
-	void Defrag(uint32 InMaxCopySize, uint32& CurrentCopySize);
+	RHICORE_API void Defrag(uint32 InMaxCopySize, uint32& CurrentCopySize);
 
 	// Stats
-	void UpdateMemoryStats(uint32& IOMemoryAllocated, uint32& IOMemoryUsed, uint32& IOMemoryFree, uint32& IOMemoryEndFree, uint32& IOAlignmentWaste, uint32& IOAllocatedPageCount, uint32& IOFullPageCount);
+	RHICORE_API void UpdateMemoryStats(uint32& IOMemoryAllocated, uint32& IOMemoryUsed, uint32& IOMemoryFree, uint32& IOMemoryEndFree, uint32& IOAlignmentWaste, uint32& IOAllocatedPageCount, uint32& IOFullPageCount);
 
 protected:
 		
 	// Supported allocator operations
-	bool TryAllocateInternal(uint32 InSizeInBytes, uint32 InAllocationAlignment, ERHIPoolResourceTypes InAllocationResourceType, FRHIPoolAllocationData& AllocationData);
-	void DeallocateInternal(FRHIPoolAllocationData& AllocationData);
+	RHICORE_API bool TryAllocateInternal(uint32 InSizeInBytes, uint32 InAllocationAlignment, ERHIPoolResourceTypes InAllocationResourceType, FRHIPoolAllocationData& AllocationData);
+	RHICORE_API void DeallocateInternal(FRHIPoolAllocationData& AllocationData);
 	
 	// Helper function to create a new platform specific pool
 	virtual FRHIMemoryPool* CreateNewPool(int16 InPoolIndex, uint32 InMinimumAllocationSize, ERHIPoolResourceTypes InAllocationResourceType) = 0;
diff --git a/Engine/Source/Runtime/RSA/Public/RSA.h b/Engine/Source/Runtime/RSA/Public/RSA.h
--- a/Engine/Source/Runtime/RSA/Public/RSA.h
+++ b/Engine/Source/Runtime/RSA/Public/RSA.h
@@ -7,46 +7,46 @@
 #include "Misc/IEngineCrypto.h"
 #include "Templates/SharedPointer.h"
 
-struct RSA_API FRSA
+struct FRSA
 {
 	/**
 	 * Create a new RSA public/private key from the supplied exponents and modulus binary data. Each of these arrays should contain a single little endian
 	 * large integer value
 	 */
-	static FRSAKeyHandle CreateKey(const TArray<uint8>& InPublicExponent, const TArray<uint8>& InPrivateExponent, const TArray<uint8>& InModulus);
+	static RSA_API FRSAKeyHandle CreateKey(const TArray<uint8>& InPublicExponent, const TArray<uint8>& InPrivateExponent, const TArray<uint8>& InModulus);
 
 	/**
 	 * Destroy the supplied key
 	 */
-	static void DestroyKey(const FRSAKeyHandle InKey);
+	static RSA_API void DestroyKey(const FRSAKeyHandle InKey);
 
 	/**
 	 * Returns the size in bytes of the supplied key
 	 */
-	static int32 GetKeySize(const FRSAKeyHandle InKey);
+	static RSA_API int32 GetKeySize(const FRSAKeyHandle InKey);
 
 	/**
 	 * Returns the maximum number of bytes that can be encrypted in a single payload
 	 */
-	static int32 GetMaxDataSize(const FRSAKeyHandle InKey);
+	static RSA_API int32 GetMaxDataSize(const FRSAKeyHandle InKey);
 
 	/**
 	 * Encrypt the supplied byte data using the given public key
 	 */
-	static int32 EncryptPublic(const TArrayView<const uint8> InSource, TArray<uint8>& OutDestination, const FRSAKeyHandle InKey);
+	static RSA_API int32 EncryptPublic(const TArrayView<const uint8> InSource, TArray<uint8>& OutDestination, const FRSAKeyHandle InKey);
 
 	/**
 	 * Encrypt the supplied byte data using the given private key
 	 */
-	static int32 EncryptPrivate(const TArrayView<const uint8> InSource, TArray<uint8>& OutDestination, const FRSAKeyHandle InKey);
+	static RSA_API int32 EncryptPrivate(const TArrayView<const uint8> InSource, TArray<uint8>& OutDestination, const FRSAKeyHandle InKey);
 
 	/**
 	 * Decrypt the supplied byte data using the given public key
 	 */
-	static int32 DecryptPublic(const TArrayView<const uint8> InSource, TArray<uint8>& OutDestination, const FRSAKeyHandle InKey);
+	static RSA_API int32 DecryptPublic(const TArrayView<const uint8> InSource, TArray<uint8>& OutDestination, const FRSAKeyHandle InKey);
 
 	/**
 	 * Encrypt the supplied byte data using the given private key
 	 */
-	static int32 DecryptPrivate(const TArrayView<const uint8> InSource, TArray<uint8>& OutDestination, const FRSAKeyHandle InKey);
-};
\ No newline at end of file
+	static RSA_API int32 DecryptPrivate(const TArrayView<const uint8> InSource, TArray<uint8>& OutDestination, const FRSAKeyHandle InKey);
+};
diff --git a/Engine/Source/Runtime/RawMesh/Public/RawMesh.h b/Engine/Source/Runtime/RawMesh/Public/RawMesh.h
--- a/Engine/Source/Runtime/RawMesh/Public/RawMesh.h
+++ b/Engine/Source/Runtime/RawMesh/Public/RawMesh.h
@@ -139,7 +139,7 @@ public:
 
 	/** Uses a hash as the GUID, useful to prevent creating new GUIDs on load for legacy assets. */
 	RAWMESH_API void UseHashAsGuid(class UObject* Owner);
-	RAWMESH_API bool UsesHashAsGuid() const { return bGuidIsHash; }
+	bool UsesHashAsGuid() const { return bGuidIsHash; }
 
 	RAWMESH_API const FByteBulkData& GetBulkData() const;
 
diff --git a/Engine/Source/Runtime/Renderer/Public/ShaderPrintParameters.h b/Engine/Source/Runtime/Renderer/Public/ShaderPrintParameters.h
--- a/Engine/Source/Runtime/Renderer/Public/ShaderPrintParameters.h
+++ b/Engine/Source/Runtime/Renderer/Public/ShaderPrintParameters.h
@@ -156,7 +156,7 @@ namespace ShaderPrint
 	//     InitShaderPrintContext Ctx = InitShaderPrintContext(...);
 	//     PrintMyVariable(Ctx, StringId, FontWhite);
 	//   }
-	struct RENDERER_API FStrings
+	struct FStrings
 	{
 	public:
 		BEGIN_SHADER_PARAMETER_STRUCT(FShaderParameters, )
@@ -166,10 +166,10 @@ namespace ShaderPrint
 			SHADER_PARAMETER_RDG_BUFFER_SRV(Buffer<uint8>, CharBuffer)
 		END_SHADER_PARAMETER_STRUCT()
 
-		FStrings(uint32 InAvgEntryCount=128u, uint32 InAvgStringLength=32u);
-		void Add(const FString& In, uint32 EntryID);
-		void Add(const FString& In);
-		FShaderParameters GetParameters(FRDGBuilder& GraphBuilder);
+		RENDERER_API FStrings(uint32 InAvgEntryCount=128u, uint32 InAvgStringLength=32u);
+		RENDERER_API void Add(const FString& In, uint32 EntryID);
+		RENDERER_API void Add(const FString& In);
+		RENDERER_API FShaderParameters GetParameters(FRDGBuilder& GraphBuilder);
 	private:
 		struct FEntryInfo
 		{
@@ -212,4 +212,4 @@ struct FFrozenShaderPrintData
 
 	TRefCountPtr<FRDGPooledBuffer> ShaderPrintEntryBuffer;
 	TRefCountPtr<FRDGPooledBuffer> ShaderPrintStateBuffer;
-};
\ No newline at end of file
+};
diff --git a/Engine/Source/Runtime/SandboxFile/Public/IPlatformFileSandboxWrapper.h b/Engine/Source/Runtime/SandboxFile/Public/IPlatformFileSandboxWrapper.h
--- a/Engine/Source/Runtime/SandboxFile/Public/IPlatformFileSandboxWrapper.h
+++ b/Engine/Source/Runtime/SandboxFile/Public/IPlatformFileSandboxWrapper.h
@@ -21,7 +21,7 @@ class IMappedFileHandle;
 **/
 DECLARE_LOG_CATEGORY_EXTERN(SandboxFile, Log, All);
 
-class SANDBOXFILE_API FSandboxPlatformFile : public IPlatformFile
+class FSandboxPlatformFile : public IPlatformFile
 {
 	/** Wrapped file */
 	IPlatformFile*		LowerLevel;
@@ -72,7 +72,7 @@ class SANDBOXFILE_API FSandboxPlatformFile : public IPlatformFile
 	 * @param AbsolutePath Absolute path to the folder to wipe
 	 * @return true if the folder's contents could be deleted, false otherwise
 	 */
-	bool WipeSandboxFolder( const TCHAR* AbsolutePath );	
+	SANDBOXFILE_API bool WipeSandboxFolder( const TCHAR* AbsolutePath );	
 
 	/**
 	 * Finds all files or folders in the given directory.
@@ -84,7 +84,7 @@ class SANDBOXFILE_API FSandboxPlatformFile : public IPlatformFile
 	 * @param Files true to include files in the Result
 	 * @param Files true to include directories in the Result
 	 */
-	void FindFiles( TArray<FString>& Result, const TCHAR* InFilename, bool Files, bool Directories );
+	SANDBOXFILE_API void FindFiles( TArray<FString>& Result, const TCHAR* InFilename, bool Files, bool Directories );
 	
 	/** Allow IPlatformFile::FindFiles */
 	using IPlatformFile::FindFiles;
@@ -98,7 +98,7 @@ class SANDBOXFILE_API FSandboxPlatformFile : public IPlatformFile
 	 * @param Tree true to recursively delete the directory and its contents
 	 * @return true if the operaton was successful.
 	 */
-	bool DeleteDirectory( const TCHAR* Path, bool Tree );
+	SANDBOXFILE_API bool DeleteDirectory( const TCHAR* Path, bool Tree );
 
 	/**
 	 * Check if a file or directory has been filtered, and hence is unavailable to the outside world
@@ -107,7 +107,7 @@ class SANDBOXFILE_API FSandboxPlatformFile : public IPlatformFile
 	 * @param bIsDirectory if true, this is a directory
 	 * @return true if it is ok to access the non-sandboxed files here
 	 */
-	bool OkForInnerAccess(const TCHAR* InFilenameOrDirectoryName, bool bIsDirectory = false) const;
+	SANDBOXFILE_API bool OkForInnerAccess(const TCHAR* InFilenameOrDirectoryName, bool bIsDirectory = false) const;
 
 	static const TCHAR* GetTypeName()
 	{
@@ -118,12 +118,12 @@ class SANDBOXFILE_API FSandboxPlatformFile : public IPlatformFile
 	 * Converts passed in filename to use a sandbox path.
 	 * @param	bEntireEngineWillUseThisSandbox		If true, the we set up the engine so that subprocesses also use this subdirectory
 	 */
-	FSandboxPlatformFile(bool bInEntireEngineWillUseThisSandbox);
+	SANDBOXFILE_API FSandboxPlatformFile(bool bInEntireEngineWillUseThisSandbox);
 
 public:
-	static TUniquePtr<FSandboxPlatformFile> Create(bool bInEntireEngineWillUseThisSandbox);
+	static SANDBOXFILE_API TUniquePtr<FSandboxPlatformFile> Create(bool bInEntireEngineWillUseThisSandbox);
 
-	virtual ~FSandboxPlatformFile();
+	SANDBOXFILE_API virtual ~FSandboxPlatformFile();
 
 	//~ For visibility of overloads we don't override
 	using IPlatformFile::IterateDirectory;
@@ -136,21 +136,21 @@ public:
 	 *
 	 *	@param	bInEnabled		true to enable the sandbox, false to disable it
 	 */
-	virtual void SetSandboxEnabled(bool bInEnabled) override;
+	SANDBOXFILE_API virtual void SetSandboxEnabled(bool bInEnabled) override;
 
 	/**
 	 *	Returns whether the sandbox is enabled or not
 	 *
 	 *	@return	bool			true if enabled, false if not
 	 */
-	virtual bool			IsSandboxEnabled() const override;
+	SANDBOXFILE_API virtual bool			IsSandboxEnabled() const override;
 
-	virtual bool			ShouldBeUsed(IPlatformFile* Inner, const TCHAR* CmdLine) const override;
-	virtual bool			Initialize(IPlatformFile* Inner, const TCHAR* CmdLine) override;
+	SANDBOXFILE_API virtual bool			ShouldBeUsed(IPlatformFile* Inner, const TCHAR* CmdLine) const override;
+	SANDBOXFILE_API virtual bool			Initialize(IPlatformFile* Inner, const TCHAR* CmdLine) override;
 
-	virtual IPlatformFile*	GetLowerLevel() override;
-	virtual void			SetLowerLevel(IPlatformFile* NewLowerLevel) override;
-	virtual const TCHAR*	GetName() const override;
+	SANDBOXFILE_API virtual IPlatformFile*	GetLowerLevel() override;
+	SANDBOXFILE_API virtual void			SetLowerLevel(IPlatformFile* NewLowerLevel) override;
+	SANDBOXFILE_API virtual const TCHAR*	GetName() const override;
 
 	/**
 	 * Converts passed in filename to use a sandbox path.
@@ -159,23 +159,23 @@ public:
 	 * 
 	 * @return	filename using sandbox path
 	 */
-	FString ConvertToSandboxPath(const TCHAR* Filename) const;
-	FString ConvertFromSandboxPath(const TCHAR* Filename) const;
+	SANDBOXFILE_API FString ConvertToSandboxPath(const TCHAR* Filename) const;
+	SANDBOXFILE_API FString ConvertFromSandboxPath(const TCHAR* Filename) const;
 
 	/** Returns sandbox directory */
-	const FString& GetSandboxDirectory() const;
+	SANDBOXFILE_API const FString& GetSandboxDirectory() const;
 
 	/** Returns the name of the sandbox directory for the game's content */
-	const FString& GetGameSandboxDirectoryName();
+	SANDBOXFILE_API const FString& GetGameSandboxDirectoryName();
 
 	/** Returns absolute root directory */
-	const FString& GetAbsoluteRootDirectory() const;
+	SANDBOXFILE_API const FString& GetAbsoluteRootDirectory() const;
 
 	/** Returns absolute game directory */
-	const FString& GetAbsoluteGameDirectory();
+	SANDBOXFILE_API const FString& GetAbsoluteGameDirectory();
 
 	/** Returns absolute path to game directory (without the game directory itself) */
-	const FString& GetAbsolutePathToGameDirectory();
+	SANDBOXFILE_API const FString& GetAbsolutePathToGameDirectory();
 
 	/** 
 	 * Add exclusion. These files and / or directories pretend not to exist so that they cannot be accessed at all (except in the sandbox) 
@@ -183,47 +183,47 @@ public:
 	 * @param bIsDirectory if true, this is a directory
 	 * @Caution, these have a performance cost
 	*/
-	void AddExclusion(const TCHAR* Wildcard, bool bIsDirectory = false);
-	void RemoveExclusion(const TCHAR* Wildcard, bool bIsDirectory = false);
+	SANDBOXFILE_API void AddExclusion(const TCHAR* Wildcard, bool bIsDirectory = false);
+	SANDBOXFILE_API void RemoveExclusion(const TCHAR* Wildcard, bool bIsDirectory = false);
 
 	/** Whether access is restricted the the sandbox or not. */
-	void SetSandboxOnly(bool bInSandboxOnly);
+	SANDBOXFILE_API void SetSandboxOnly(bool bInSandboxOnly);
 
 	// IPlatformFile Interface
 
-	virtual bool					FileExists(const TCHAR* Filename) override;
-	virtual int64					FileSize(const TCHAR* Filename) override;
-	virtual bool					DeleteFile(const TCHAR* Filename) override;
-	virtual bool					IsReadOnly(const TCHAR* Filename) override;
-	virtual bool					MoveFile(const TCHAR* To, const TCHAR* From) override;
-	virtual bool					SetReadOnly(const TCHAR* Filename, bool bNewReadOnlyValue) override;
-	virtual FDateTime				GetTimeStamp(const TCHAR* Filename) override;
-	virtual void 					SetTimeStamp(const TCHAR* Filename, FDateTime DateTime) override;
-	virtual FDateTime				GetAccessTimeStamp(const TCHAR* Filename) override;
-	virtual FString					GetFilenameOnDisk(const TCHAR* Filename) override;
-	virtual IFileHandle*			OpenRead(const TCHAR* Filename, bool bAllowWrite = false) override;
-	virtual IFileHandle*			OpenWrite(const TCHAR* Filename, bool bAppend = false, bool bAllowRead = false) override;
-	virtual bool					DirectoryExists(const TCHAR* Directory) override;
-	virtual bool					CreateDirectory(const TCHAR* Directory) override;
-	virtual bool					DeleteDirectory(const TCHAR* Directory) override;
-	virtual FFileStatData			GetStatData(const TCHAR* FilenameOrDirectory) override;
-
-	virtual bool					IterateDirectory(const TCHAR* Directory, IPlatformFile::FDirectoryVisitor& Visitor) override;
-	virtual bool					IterateDirectoryRecursively(const TCHAR* Directory, IPlatformFile::FDirectoryVisitor& Visitor) override;
-
-	virtual bool					IterateDirectoryStat(const TCHAR* Directory, IPlatformFile::FDirectoryStatVisitor& Visitor) override;
-	virtual bool					IterateDirectoryStatRecursively(const TCHAR* Directory, IPlatformFile::FDirectoryStatVisitor& Visitor) override;
-
-	virtual bool					DeleteDirectoryRecursively(const TCHAR* Directory) override;
-	virtual bool					CreateDirectoryTree(const TCHAR* Directory) override;
-
-	virtual bool					CopyFile(const TCHAR* To, const TCHAR* From, EPlatformFileRead ReadFlags = EPlatformFileRead::None, EPlatformFileWrite WriteFlags = EPlatformFileWrite::None) override;
-	virtual FString					ConvertToAbsolutePathForExternalAppForRead( const TCHAR* Filename ) override;
-	virtual FString					ConvertToAbsolutePathForExternalAppForWrite( const TCHAR* Filename ) override;
-
-	virtual IAsyncReadFileHandle*	OpenAsyncRead(const TCHAR* Filename) override;
-	virtual void					SetAsyncMinimumPriority(EAsyncIOPriorityAndFlags Priority) override;
-	virtual IMappedFileHandle*		OpenMapped(const TCHAR* Filename) override;
+	SANDBOXFILE_API virtual bool					FileExists(const TCHAR* Filename) override;
+	SANDBOXFILE_API virtual int64					FileSize(const TCHAR* Filename) override;
+	SANDBOXFILE_API virtual bool					DeleteFile(const TCHAR* Filename) override;
+	SANDBOXFILE_API virtual bool					IsReadOnly(const TCHAR* Filename) override;
+	SANDBOXFILE_API virtual bool					MoveFile(const TCHAR* To, const TCHAR* From) override;
+	SANDBOXFILE_API virtual bool					SetReadOnly(const TCHAR* Filename, bool bNewReadOnlyValue) override;
+	SANDBOXFILE_API virtual FDateTime				GetTimeStamp(const TCHAR* Filename) override;
+	SANDBOXFILE_API virtual void 					SetTimeStamp(const TCHAR* Filename, FDateTime DateTime) override;
+	SANDBOXFILE_API virtual FDateTime				GetAccessTimeStamp(const TCHAR* Filename) override;
+	SANDBOXFILE_API virtual FString					GetFilenameOnDisk(const TCHAR* Filename) override;
+	SANDBOXFILE_API virtual IFileHandle*			OpenRead(const TCHAR* Filename, bool bAllowWrite = false) override;
+	SANDBOXFILE_API virtual IFileHandle*			OpenWrite(const TCHAR* Filename, bool bAppend = false, bool bAllowRead = false) override;
+	SANDBOXFILE_API virtual bool					DirectoryExists(const TCHAR* Directory) override;
+	SANDBOXFILE_API virtual bool					CreateDirectory(const TCHAR* Directory) override;
+	SANDBOXFILE_API virtual bool					DeleteDirectory(const TCHAR* Directory) override;
+	SANDBOXFILE_API virtual FFileStatData			GetStatData(const TCHAR* FilenameOrDirectory) override;
+
+	SANDBOXFILE_API virtual bool					IterateDirectory(const TCHAR* Directory, IPlatformFile::FDirectoryVisitor& Visitor) override;
+	SANDBOXFILE_API virtual bool					IterateDirectoryRecursively(const TCHAR* Directory, IPlatformFile::FDirectoryVisitor& Visitor) override;
+
+	SANDBOXFILE_API virtual bool					IterateDirectoryStat(const TCHAR* Directory, IPlatformFile::FDirectoryStatVisitor& Visitor) override;
+	SANDBOXFILE_API virtual bool					IterateDirectoryStatRecursively(const TCHAR* Directory, IPlatformFile::FDirectoryStatVisitor& Visitor) override;
+
+	SANDBOXFILE_API virtual bool					DeleteDirectoryRecursively(const TCHAR* Directory) override;
+	SANDBOXFILE_API virtual bool					CreateDirectoryTree(const TCHAR* Directory) override;
+
+	SANDBOXFILE_API virtual bool					CopyFile(const TCHAR* To, const TCHAR* From, EPlatformFileRead ReadFlags = EPlatformFileRead::None, EPlatformFileWrite WriteFlags = EPlatformFileWrite::None) override;
+	SANDBOXFILE_API virtual FString					ConvertToAbsolutePathForExternalAppForRead( const TCHAR* Filename ) override;
+	SANDBOXFILE_API virtual FString					ConvertToAbsolutePathForExternalAppForWrite( const TCHAR* Filename ) override;
+
+	SANDBOXFILE_API virtual IAsyncReadFileHandle*	OpenAsyncRead(const TCHAR* Filename) override;
+	SANDBOXFILE_API virtual void					SetAsyncMinimumPriority(EAsyncIOPriorityAndFlags Priority) override;
+	SANDBOXFILE_API virtual IMappedFileHandle*		OpenMapped(const TCHAR* Filename) override;
 
 	friend class FSandboxVisitor;
 	friend class FSandboxStatVisitor;
diff --git a/Engine/Source/Runtime/Serialization/Public/Backends/CborStructDeserializerBackend.h b/Engine/Source/Runtime/Serialization/Public/Backends/CborStructDeserializerBackend.h
--- a/Engine/Source/Runtime/Serialization/Public/Backends/CborStructDeserializerBackend.h
+++ b/Engine/Source/Runtime/Serialization/Public/Backends/CborStructDeserializerBackend.h
@@ -16,7 +16,7 @@ class FProperty;
 /**
  * Implements a reader for UStruct deserialization using Cbor.
  */
-class SERIALIZATION_API FCborStructDeserializerBackend
+class FCborStructDeserializerBackend
 	: public IStructDeserializerBackend
 {
 public:
@@ -27,20 +27,20 @@ public:
 	 * @param CborDataEndianness The CBOR data endianness stored in the archive.
 	 * @note For backward compatibility and performance, the implementation default to the the platform endianness rather than the CBOR standard one (big endian).
 	 */
-	FCborStructDeserializerBackend(FArchive& Archive, ECborEndianness CborDataEndianness = ECborEndianness::Platform, bool bInIsLWCCompatibilityMode = false);
-	virtual ~FCborStructDeserializerBackend();
+	SERIALIZATION_API FCborStructDeserializerBackend(FArchive& Archive, ECborEndianness CborDataEndianness = ECborEndianness::Platform, bool bInIsLWCCompatibilityMode = false);
+	SERIALIZATION_API virtual ~FCborStructDeserializerBackend();
 
 public:
 
 	// IStructDeserializerBackend interface
-	virtual const FString& GetCurrentPropertyName() const override;
-	virtual FString GetDebugString() const override;
-	virtual const FString& GetLastErrorMessage() const override;
-	virtual bool GetNextToken(EStructDeserializerBackendTokens& OutToken) override;
-	virtual bool ReadProperty(FProperty* Property, FProperty* Outer, void* Data, int32 ArrayIndex) override;
-	virtual bool ReadPODArray(FArrayProperty* ArrayProperty, void* Data) override;
-	virtual void SkipArray() override;
-	virtual void SkipStructure() override;
+	SERIALIZATION_API virtual const FString& GetCurrentPropertyName() const override;
+	SERIALIZATION_API virtual FString GetDebugString() const override;
+	SERIALIZATION_API virtual const FString& GetLastErrorMessage() const override;
+	SERIALIZATION_API virtual bool GetNextToken(EStructDeserializerBackendTokens& OutToken) override;
+	SERIALIZATION_API virtual bool ReadProperty(FProperty* Property, FProperty* Outer, void* Data, int32 ArrayIndex) override;
+	SERIALIZATION_API virtual bool ReadPODArray(FArrayProperty* ArrayProperty, void* Data) override;
+	SERIALIZATION_API virtual void SkipArray() override;
+	SERIALIZATION_API virtual void SkipStructure() override;
 
 private:
 	/** Holds the Cbor reader used for the actual reading of the archive. */
diff --git a/Engine/Source/Runtime/Serialization/Public/Backends/CborStructSerializerBackend.h b/Engine/Source/Runtime/Serialization/Public/Backends/CborStructSerializerBackend.h
--- a/Engine/Source/Runtime/Serialization/Public/Backends/CborStructSerializerBackend.h
+++ b/Engine/Source/Runtime/Serialization/Public/Backends/CborStructSerializerBackend.h
@@ -14,7 +14,7 @@ class FArchive;
 /**
  * Implements a writer for UStruct serialization using Cbor.
  */
-class SERIALIZATION_API FCborStructSerializerBackend
+class FCborStructSerializerBackend
 	: public IStructSerializerBackend
 {
 public:
@@ -26,7 +26,7 @@ public:
 	 * @param InArchive The archive to serialize into.
 	 */
 	UE_DEPRECATED(4.22, "Use the two-parameter constructor with EStructSerializerBackendFlags::Legacy only if you need backwards compatibility with code compiled prior to 4.22; otherwise use EStructSerializerBackendFlags::Default.")
-	FCborStructSerializerBackend(FArchive& InArchive);
+	SERIALIZATION_API FCborStructSerializerBackend(FArchive& InArchive);
 
 	/**
 	 * Creates and initializes a new instance with the given flags.
@@ -34,20 +34,20 @@ public:
 	 * @param InArchive The archive to serialize into.
 	 * @param InFlags The flags that control the serialization behavior (typically EStructSerializerBackendFlags::Default).
 	 */
-	FCborStructSerializerBackend(FArchive& InArchive, const EStructSerializerBackendFlags InFlags);
+	SERIALIZATION_API FCborStructSerializerBackend(FArchive& InArchive, const EStructSerializerBackendFlags InFlags);
 
-	virtual ~FCborStructSerializerBackend();
+	SERIALIZATION_API virtual ~FCborStructSerializerBackend();
 
 public:
 
 	// IStructSerializerBackend interface
-	virtual void BeginArray(const FStructSerializerState& State) override;
-	virtual void BeginStructure(const FStructSerializerState& State) override;
-	virtual void EndArray(const FStructSerializerState& State) override;
-	virtual void EndStructure(const FStructSerializerState& State) override;
-	virtual void WriteComment(const FString& Comment) override;
-	virtual void WriteProperty(const FStructSerializerState& State, int32 ArrayIndex = 0) override;
-	virtual bool WritePODArray(const FStructSerializerState& State) override;
+	SERIALIZATION_API virtual void BeginArray(const FStructSerializerState& State) override;
+	SERIALIZATION_API virtual void BeginStructure(const FStructSerializerState& State) override;
+	SERIALIZATION_API virtual void EndArray(const FStructSerializerState& State) override;
+	SERIALIZATION_API virtual void EndStructure(const FStructSerializerState& State) override;
+	SERIALIZATION_API virtual void WriteComment(const FString& Comment) override;
+	SERIALIZATION_API virtual void WriteProperty(const FStructSerializerState& State, int32 ArrayIndex = 0) override;
+	SERIALIZATION_API virtual bool WritePODArray(const FStructSerializerState& State) override;
 
 private:
 	/** Holds the Cbor writer used for the actual serialization. */
diff --git a/Engine/Source/Runtime/Serialization/Public/Backends/JsonStructDeserializerBackend.h b/Engine/Source/Runtime/Serialization/Public/Backends/JsonStructDeserializerBackend.h
--- a/Engine/Source/Runtime/Serialization/Public/Backends/JsonStructDeserializerBackend.h
+++ b/Engine/Source/Runtime/Serialization/Public/Backends/JsonStructDeserializerBackend.h
@@ -21,7 +21,7 @@ class FProperty;
  * all based on templates. At some point we will refactor the low-level Json API to provide more
  * flexibility for serialization.
  */
-class SERIALIZATION_API FJsonStructDeserializerBackend
+class FJsonStructDeserializerBackend
 	: public IStructDeserializerBackend
 {
 public:
@@ -39,13 +39,13 @@ public:
 
 	// IStructDeserializerBackend interface
 
-	virtual const FString& GetCurrentPropertyName() const override;
-	virtual FString GetDebugString() const override;
-	virtual const FString& GetLastErrorMessage() const override;
-	virtual bool GetNextToken( EStructDeserializerBackendTokens& OutToken ) override;
-	virtual bool ReadProperty( FProperty* Property, FProperty* Outer, void* Data, int32 ArrayIndex ) override;
-	virtual void SkipArray() override;
-	virtual void SkipStructure() override;
+	SERIALIZATION_API virtual const FString& GetCurrentPropertyName() const override;
+	SERIALIZATION_API virtual FString GetDebugString() const override;
+	SERIALIZATION_API virtual const FString& GetLastErrorMessage() const override;
+	SERIALIZATION_API virtual bool GetNextToken( EStructDeserializerBackendTokens& OutToken ) override;
+	SERIALIZATION_API virtual bool ReadProperty( FProperty* Property, FProperty* Outer, void* Data, int32 ArrayIndex ) override;
+	SERIALIZATION_API virtual void SkipArray() override;
+	SERIALIZATION_API virtual void SkipStructure() override;
 
 protected:
 	FString& GetLastIdentifier()
diff --git a/Engine/Source/Runtime/Serialization/Public/Backends/JsonStructSerializerBackend.h b/Engine/Source/Runtime/Serialization/Public/Backends/JsonStructSerializerBackend.h
--- a/Engine/Source/Runtime/Serialization/Public/Backends/JsonStructSerializerBackend.h
+++ b/Engine/Source/Runtime/Serialization/Public/Backends/JsonStructSerializerBackend.h
@@ -22,7 +22,7 @@ class FArchive;
  * all based on templates. At some point we will refactor the low-level Json API to provide more
  * flexibility for serialization.
  */
-class SERIALIZATION_API FJsonStructSerializerBackend
+class FJsonStructSerializerBackend
 	: public IStructSerializerBackend
 {
 public:
@@ -54,12 +54,12 @@ public:
 
 	// IStructSerializerBackend interface
 
-	virtual void BeginArray(const FStructSerializerState& State) override;
-	virtual void BeginStructure(const FStructSerializerState& State) override;
-	virtual void EndArray(const FStructSerializerState& State) override;
-	virtual void EndStructure(const FStructSerializerState& State) override;
-	virtual void WriteComment(const FString& Comment) override;
-	virtual void WriteProperty(const FStructSerializerState& State, int32 ArrayIndex = 0) override;
+	SERIALIZATION_API virtual void BeginArray(const FStructSerializerState& State) override;
+	SERIALIZATION_API virtual void BeginStructure(const FStructSerializerState& State) override;
+	SERIALIZATION_API virtual void EndArray(const FStructSerializerState& State) override;
+	SERIALIZATION_API virtual void EndStructure(const FStructSerializerState& State) override;
+	SERIALIZATION_API virtual void WriteComment(const FString& Comment) override;
+	SERIALIZATION_API virtual void WriteProperty(const FStructSerializerState& State, int32 ArrayIndex = 0) override;
 
 protected:
 
diff --git a/Engine/Source/Runtime/SignalProcessing/Public/DSP/AlignedBlockBuffer.h b/Engine/Source/Runtime/SignalProcessing/Public/DSP/AlignedBlockBuffer.h
--- a/Engine/Source/Runtime/SignalProcessing/Public/DSP/AlignedBlockBuffer.h
+++ b/Engine/Source/Runtime/SignalProcessing/Public/DSP/AlignedBlockBuffer.h
@@ -11,51 +11,51 @@ namespace Audio
 	// for audio applications by ensuring memory allocations only occur in the constructor or when capacity 
 	// is altered, memory freeing only happens in the destructor or when capacity is altered. All assignments 
 	// are performed using block memory copies.
-	class SIGNALPROCESSING_API FAlignedBlockBuffer
+	class FAlignedBlockBuffer
 	{
 	public:
 		
 		// Constructor
-		FAlignedBlockBuffer(int32 InSampleCapacity=8192, int32 InMaxNumInspectSamples=1024, uint32 InByteAlignment=AUDIO_SIMD_BYTE_ALIGNMENT, uint32 InAllocByteAlignment=AUDIO_BUFFER_ALIGNMENT);
+		SIGNALPROCESSING_API FAlignedBlockBuffer(int32 InSampleCapacity=8192, int32 InMaxNumInspectSamples=1024, uint32 InByteAlignment=AUDIO_SIMD_BYTE_ALIGNMENT, uint32 InAllocByteAlignment=AUDIO_BUFFER_ALIGNMENT);
 
 		// Destructor
-		~FAlignedBlockBuffer() throw();
+		SIGNALPROCESSING_API ~FAlignedBlockBuffer() throw();
 		
 		// Maximum size of buffer
-		int32 GetSampleCapacity() const;
+		SIGNALPROCESSING_API int32 GetSampleCapacity() const;
 		
 		// Amount of samples left over in buffer
-		int32 GetNumAvailable() const;
+		SIGNALPROCESSING_API int32 GetNumAvailable() const;
 
 		// Add zeros to the buffer
-		void AddZeros(int32 InNum);
+		SIGNALPROCESSING_API void AddZeros(int32 InNum);
 		
 		// Add samples to the buffer.
-		void AddSamples(const float* InSamples, int32 InNum);
+		SIGNALPROCESSING_API void AddSamples(const float* InSamples, int32 InNum);
 		
 		// Remove samples from the buffer
-		void RemoveSamples(int32 InNum);
+		SIGNALPROCESSING_API void RemoveSamples(int32 InNum);
 
 		// Clear entire buffer
-		void ClearSamples();
+		SIGNALPROCESSING_API void ClearSamples();
 
 		// Inspect samples in the buffer. Returned samples abide by the memory alignment specified during construction.
 		// Warning: Not thread safe. The returned pointer is only valid as long as this object exists 
 		// and before any mutating operations are made to this object.
-		const float* InspectSamples(int32 InNum, int32 InOffset=0);
+		SIGNALPROCESSING_API const float* InspectSamples(int32 InNum, int32 InOffset=0);
 		
 		// Maximum number of elements that can be inspected
-		int32 GetMaxNumInspectSamples() const;
+		SIGNALPROCESSING_API int32 GetMaxNumInspectSamples() const;
 		
 	protected:
 
 		// Set the maximum number of elements that can be inspected
-		void SetMaxNumInspectSamples(int32 InMax);
+		SIGNALPROCESSING_API void SetMaxNumInspectSamples(int32 InMax);
 
-		float* AllocateAlignedFloatArray(int32 InNum) const;
-		void FreeFloatArray(float*& Array) const;
-		void ZeroFloatArray(float* InArray, int32 InNum) const;
-		void CopyFloatArray(float* ToArray, const float* FromArray, int32 InNum) const;
+		SIGNALPROCESSING_API float* AllocateAlignedFloatArray(int32 InNum) const;
+		SIGNALPROCESSING_API void FreeFloatArray(float*& Array) const;
+		SIGNALPROCESSING_API void ZeroFloatArray(float* InArray, int32 InNum) const;
+		SIGNALPROCESSING_API void CopyFloatArray(float* ToArray, const float* FromArray, int32 InNum) const;
 
 	private:
 
diff --git a/Engine/Source/Runtime/SignalProcessing/Public/DSP/AllPassFilter.h b/Engine/Source/Runtime/SignalProcessing/Public/DSP/AllPassFilter.h
--- a/Engine/Source/Runtime/SignalProcessing/Public/DSP/AllPassFilter.h
+++ b/Engine/Source/Runtime/SignalProcessing/Public/DSP/AllPassFilter.h
@@ -9,20 +9,20 @@ namespace Audio
 {
 	// Implementation of a delay line with a feedback/feedforward gain coefficient
 	// APF filters pass all frequencies but changes phase relationships of frequencies
-	class SIGNALPROCESSING_API FDelayAPF : public FDelay
+	class FDelayAPF : public FDelay
 	{
 	public:
 		// Constructor
-		FDelayAPF();
+		SIGNALPROCESSING_API FDelayAPF();
 
 		// Destructor
-		~FDelayAPF();
+		SIGNALPROCESSING_API ~FDelayAPF();
 
 		// Set the APF feedback/feedforward gain coefficient
 		void SetG(float InG) { G = InG; }
 
 		// overrides
-		virtual float ProcessAudioSample(const float pInput) override;
+		SIGNALPROCESSING_API virtual float ProcessAudioSample(const float pInput) override;
 
 	protected:
 		// Feedback/Feedforward gain coefficient
diff --git a/Engine/Source/Runtime/SignalProcessing/Public/DSP/AllPassFractionalDelay.h b/Engine/Source/Runtime/SignalProcessing/Public/DSP/AllPassFractionalDelay.h
--- a/Engine/Source/Runtime/SignalProcessing/Public/DSP/AllPassFractionalDelay.h
+++ b/Engine/Source/Runtime/SignalProcessing/Public/DSP/AllPassFractionalDelay.h
@@ -9,21 +9,21 @@
 namespace Audio
 {
 	// Performs a fractional delay utilzing a single order all pass filter.
-	class SIGNALPROCESSING_API FAllPassFractionalDelay
+	class FAllPassFractionalDelay
 	{
 	public:
 		// InMaxDelay sets the maximum allowable delay that this object can support. The minimum delay is 0.5
 		// InNumInternalBufferSamples sets the maximum block processing size.
-		FAllPassFractionalDelay(int32 InMaxDelay, int32 InNumInternalBufferSamples);
+		SIGNALPROCESSING_API FAllPassFractionalDelay(int32 InMaxDelay, int32 InNumInternalBufferSamples);
 
-		~FAllPassFractionalDelay();
+		SIGNALPROCESSING_API ~FAllPassFractionalDelay();
 
 		// Apply delay to InSamples. Fill OutSamples with data from the delay line at a delay of InDelay.
 		// InDelay must be equal length to InSamples. 
-		void ProcessAudio(const FAlignedFloatBuffer& InSamples, const FAlignedFloatBuffer& InDelays, FAlignedFloatBuffer& OutSamples);
+		SIGNALPROCESSING_API void ProcessAudio(const FAlignedFloatBuffer& InSamples, const FAlignedFloatBuffer& InDelays, FAlignedFloatBuffer& OutSamples);
 
 		// Set all values in internal delay line to zero. 
-		void Reset();
+		SIGNALPROCESSING_API void Reset();
 
 	private:
 		void ProcessAudioBlock(const float* InSamples, const float* InDelays, const int32 InNum, float* OutSamples);
diff --git a/Engine/Source/Runtime/SignalProcessing/Public/DSP/Amp.h b/Engine/Source/Runtime/SignalProcessing/Public/DSP/Amp.h
--- a/Engine/Source/Runtime/SignalProcessing/Public/DSP/Amp.h
+++ b/Engine/Source/Runtime/SignalProcessing/Public/DSP/Amp.h
@@ -7,7 +7,7 @@
 
 namespace Audio
 {
-	struct SIGNALPROCESSING_API FWetDry
+	struct FWetDry
 	{
 		// Level of wet signal
 		float WetLevel;
@@ -23,58 +23,58 @@ namespace Audio
 	};
 
 	// Class which manages scaling audio input and performing panning operations
-	class SIGNALPROCESSING_API FAmp
+	class FAmp
 	{
 	public:
-		FAmp();
-		~FAmp();
+		SIGNALPROCESSING_API FAmp();
+		SIGNALPROCESSING_API ~FAmp();
 
 		// Initializes the amp with the mod matrix
-		void Init(const int32 InVoiceId = INDEX_NONE, FModulationMatrix* ModMatrix = nullptr);
+		SIGNALPROCESSING_API void Init(const int32 InVoiceId = INDEX_NONE, FModulationMatrix* ModMatrix = nullptr);
 
 		// Sets the direct gain in decibel
-		void SetGainDb(const float InGainDB);
+		SIGNALPROCESSING_API void SetGainDb(const float InGainDB);
 
 		// Sets the gain modulation in decibel
-		void SetGainModDb(const float InGainModDb);
+		SIGNALPROCESSING_API void SetGainModDb(const float InGainModDb);
 
 		// Sets the direct gain in linear
-		void SetGain(const float InGainLinear);
+		SIGNALPROCESSING_API void SetGain(const float InGainLinear);
 
 		// Sets the gain modulation in linear. Expects modulation value to be bipolar (for AM synth, etc)
-		void SetGainMod(const float InBipolarGainModLinear);
+		SIGNALPROCESSING_API void SetGainMod(const float InBipolarGainModLinear);
 
 		// Sets the gain based on an envelope value
-		void SetGainEnv(const float InGainEnv);
+		SIGNALPROCESSING_API void SetGainEnv(const float InGainEnv);
 
 		// Sets the gain based on an envelope value
-		void SetGainEnvDb(const float InGainEnvDb);
+		SIGNALPROCESSING_API void SetGainEnvDb(const float InGainEnvDb);
 
 		// Change the allowed range of the gain output
-		void SetGainRange(const float InMin, const float InMax);
+		SIGNALPROCESSING_API void SetGainRange(const float InMin, const float InMax);
 		
 		// Sets the gain using a midi velocity value
-		void SetVelocity(const float InVelocity);
+		SIGNALPROCESSING_API void SetVelocity(const float InVelocity);
 
 		// Sets the pan
-		void SetPan(const float InPan);
+		SIGNALPROCESSING_API void SetPan(const float InPan);
 
 		// Sets the pan modulator
-		void SetPanModulator(const float InPanMod);
+		SIGNALPROCESSING_API void SetPanModulator(const float InPanMod);
 
 		// Takes mono input and generates stereo output
-		void ProcessAudio(const float LeftIn, float* LeftOutput, float* RightOutput);
+		SIGNALPROCESSING_API void ProcessAudio(const float LeftIn, float* LeftOutput, float* RightOutput);
 
 		// Takes stereo input and generates stereo output
-		void ProcessAudio(const float LeftIn, const float RightIn, float* LeftOutput, float* RightOutput);
+		SIGNALPROCESSING_API void ProcessAudio(const float LeftIn, const float RightIn, float* LeftOutput, float* RightOutput);
 
 		// Generates a new gain value for left and right outputs. 
-		void Generate(float& OutGainLeft, float& OutGainRight);
+		SIGNALPROCESSING_API void Generate(float& OutGainLeft, float& OutGainRight);
 
-		void Reset();
+		SIGNALPROCESSING_API void Reset();
 
 		// Updates the final output left and right gain based on current settings
-		void Update();
+		SIGNALPROCESSING_API void Update();
 
 		const FPatchDestination GetModDestGainScale() const { return GainScaleDest; }
 		const FPatchDestination GetModDestGainEnv() const { return GainEnvDest; }
diff --git a/Engine/Source/Runtime/SignalProcessing/Public/DSP/AudioBufferDistanceAttenuation.h b/Engine/Source/Runtime/SignalProcessing/Public/DSP/AudioBufferDistanceAttenuation.h
--- a/Engine/Source/Runtime/SignalProcessing/Public/DSP/AudioBufferDistanceAttenuation.h
+++ b/Engine/Source/Runtime/SignalProcessing/Public/DSP/AudioBufferDistanceAttenuation.h
@@ -9,7 +9,7 @@
 namespace Audio
 {
 	// Settings for audio buffer distance attenuation
-	struct SIGNALPROCESSING_API FAudioBufferDistanceAttenuationSettings
+	struct FAudioBufferDistanceAttenuationSettings
 	{
 		FAudioBufferDistanceAttenuationSettings() = default;
 		~FAudioBufferDistanceAttenuationSettings() = default;
@@ -46,4 +46,4 @@ namespace Audio
 		const FAudioBufferDistanceAttenuationSettings& InSettings, 
 		float& InOutAttenuation);
 
-}
\ No newline at end of file
+}
diff --git a/Engine/Source/Runtime/SignalProcessing/Public/DSP/AudioChannelFormatConverter.h b/Engine/Source/Runtime/SignalProcessing/Public/DSP/AudioChannelFormatConverter.h
--- a/Engine/Source/Runtime/SignalProcessing/Public/DSP/AudioChannelFormatConverter.h
+++ b/Engine/Source/Runtime/SignalProcessing/Public/DSP/AudioChannelFormatConverter.h
@@ -10,7 +10,7 @@
 namespace Audio
 {
 	/** Inteface for Channel Format Converters which process deinterleaved audio. */
-	class SIGNALPROCESSING_API IChannelFormatConverter
+	class IChannelFormatConverter
 	{
 		public:
 			/** Description of input audio format. */
@@ -51,7 +51,7 @@ namespace Audio
 	/** FBaseChannelFormatConverter implements channel conversion using a simple
 	 * mixing matrix. 
 	 */
-	class SIGNALPROCESSING_API FBaseChannelFormatConverter : public IChannelFormatConverter
+	class FBaseChannelFormatConverter : public IChannelFormatConverter
 	{
 
 		public:
@@ -74,10 +74,10 @@ namespace Audio
 			virtual ~FBaseChannelFormatConverter() = default;
 
 			/** Return the input format handled by this converter. */
-			const FInputFormat& GetInputFormat() const override;
+			SIGNALPROCESSING_API const FInputFormat& GetInputFormat() const override;
 
 			/** Return the output format handled by this converter. */
-			const FOutputFormat& GetOutputFormat() const override;
+			SIGNALPROCESSING_API const FOutputFormat& GetOutputFormat() const override;
 
 			/** Sets the output gain scalar to apply to all output audio. 
 			 *
@@ -86,7 +86,7 @@ namespace Audio
 			 *                      the duration of a single buffer. If false,
 			 *                      gain values are applied immediately. 
 			 */
-			void SetOutputGain(float InOutputGain, bool bFadeToGain=true);
+			SIGNALPROCESSING_API void SetOutputGain(float InOutputGain, bool bFadeToGain=true);
 
 			/** Sets the gain scalar to apply to a specific input/output channel pair.
 			 *
@@ -95,7 +95,7 @@ namespace Audio
 			 *                      the duration of a single buffer. If false,
 			 *                      gain values are applied immediately. 
 			 */
-			void SetMixGain(const FChannelMixEntry& InEntry, bool bFadeToGain=true);
+			SIGNALPROCESSING_API void SetMixGain(const FChannelMixEntry& InEntry, bool bFadeToGain=true);
 
 			/** Sets the gain scalar to apply to a specific input/output channel pair.
 			 *
@@ -107,7 +107,7 @@ namespace Audio
 			 *                      the duration of a single buffer. If false,
 			 *                      gain values are applied immediately. 
 			 */
-			void SetMixGain(int32 InInputChannelIndex, int32 InOutputChannelIndex, float InGain, bool bFadeToGain=true);
+			SIGNALPROCESSING_API void SetMixGain(int32 InInputChannelIndex, int32 InOutputChannelIndex, float InGain, bool bFadeToGain=true);
 
 			/** Returns the scalar gain used to mix the input channel index into 
 			 * the output channel index.  
@@ -115,14 +115,14 @@ namespace Audio
 			 * If the gains are to be faded during the subsequent call to ProcessAudio(),
 			 * this will return the ending gain.
 			 */
-			float GetTargetMixGain(int32 InInputChannelIndex, int32 InOutputChannelIndex) const;
+			SIGNALPROCESSING_API float GetTargetMixGain(int32 InInputChannelIndex, int32 InOutputChannelIndex) const;
 
 			/** Returns the scalar gain applied to the output audio.
 			 *
 			 * If the gain is to be faded during the subsequent call to ProcessAudio(),
 			 * this will return the ending gain.
 			 */
-			float GetTargetOutputGain() const;
+			SIGNALPROCESSING_API float GetTargetOutputGain() const;
 
 			/** Converter the audio format from the FInputFormat to the FOutputFormat. 
 			 *
@@ -133,7 +133,7 @@ namespace Audio
 			 * @param InInputBuffers - An array of input audio buffers.
 			 * @param OutOutputBuffers - An array of buffers where output audio is stored. 
 			 */
-			void ProcessAudio(const TArray<FAlignedFloatBuffer>& InInputBuffers, TArray<FAlignedFloatBuffer>& OutOutputBuffers) override;
+			SIGNALPROCESSING_API void ProcessAudio(const TArray<FAlignedFloatBuffer>& InInputBuffers, TArray<FAlignedFloatBuffer>& OutOutputBuffers) override;
 
 
 			/** Create a FBaseChannelFormatConverter
@@ -145,11 +145,11 @@ namespace Audio
 			 *
 			 * @return A TUniquePtr to a FBaseChannelFormatConverter.
 			 */
-			static TUniquePtr<FBaseChannelFormatConverter> CreateBaseFormatConverter(const FInputFormat& InInputFormat, const FOutputFormat& InOutputFormat, TArrayView<const FChannelMixEntry> InMixEntries, int32 InNumFramesPerCall);
+			static SIGNALPROCESSING_API TUniquePtr<FBaseChannelFormatConverter> CreateBaseFormatConverter(const FInputFormat& InInputFormat, const FOutputFormat& InOutputFormat, TArrayView<const FChannelMixEntry> InMixEntries, int32 InNumFramesPerCall);
 
 		protected:
 
-			FBaseChannelFormatConverter(const FInputFormat& InInputFormat, const FOutputFormat& InOutputFormat, TArrayView<const FChannelMixEntry> InMixEntries, int32 InNumFramesPerCall);
+			SIGNALPROCESSING_API FBaseChannelFormatConverter(const FInputFormat& InInputFormat, const FOutputFormat& InOutputFormat, TArrayView<const FChannelMixEntry> InMixEntries, int32 InNumFramesPerCall);
 
 		private:
 
diff --git a/Engine/Source/Runtime/SignalProcessing/Public/DSP/AudioFFT.h b/Engine/Source/Runtime/SignalProcessing/Public/DSP/AudioFFT.h
--- a/Engine/Source/Runtime/SignalProcessing/Public/DSP/AudioFFT.h
+++ b/Engine/Source/Runtime/SignalProcessing/Public/DSP/AudioFFT.h
@@ -43,7 +43,7 @@ namespace Audio
 	/**
 	 * Class used to generate, contain and apply a DSP window of a given type.
 	 */
-	class SIGNALPROCESSING_API FWindow
+	class FWindow
 	{
 	public:
 		FWindow() = delete;
@@ -57,13 +57,13 @@ namespace Audio
 		 *                     Generally, set this to false if using this window with an STFT, but use true
 		 *                     if this window will be used on an entire, self-contained signal.
 		 */
-		FWindow(EWindowType InType, int32 InNumFrames, int32 InNumChannels, bool bIsPeriodic);
+		SIGNALPROCESSING_API FWindow(EWindowType InType, int32 InNumFrames, int32 InNumChannels, bool bIsPeriodic);
 
 		// Apply this window to InBuffer, which is expected to be an interleaved buffer with the same amount of frames
 		// and channels this window was constructed with.
-		void ApplyToBuffer(float* InBuffer);
+		SIGNALPROCESSING_API void ApplyToBuffer(float* InBuffer);
 
-		EWindowType GetWindowType() const;
+		SIGNALPROCESSING_API EWindowType GetWindowType() const;
 
 	private:
 		EWindowType WindowType;
@@ -71,7 +71,7 @@ namespace Audio
 		int32 NumSamples;
 
 		// Generate the window. Called on constructor.
-		void Generate(int32 NumFrames, int32 NumChannels, bool bIsPeriodic);
+		SIGNALPROCESSING_API void Generate(int32 NumFrames, int32 NumChannels, bool bIsPeriodic);
 	};
 
 	struct FFTTimeDomainData_DEPRECATED
@@ -136,25 +136,25 @@ namespace Audio
 
 	// FFT Algorithm factory for this FFT implementation
 	class UE_DEPRECATED(5.1, "Use FVectorFFTFactory instead.") FAudioFFTAlgorithmFactory;
-	class SIGNALPROCESSING_API FAudioFFTAlgorithmFactory : public IFFTAlgorithmFactory
+	class FAudioFFTAlgorithmFactory : public IFFTAlgorithmFactory
 	{
 		public:
-			virtual ~FAudioFFTAlgorithmFactory();
+			SIGNALPROCESSING_API virtual ~FAudioFFTAlgorithmFactory();
 
 			// Name of this fft algorithm factory. 
-			virtual FName GetFactoryName() const override;
+			SIGNALPROCESSING_API virtual FName GetFactoryName() const override;
 
 			// If true, this implementation uses hardware acceleration.
-			virtual bool IsHardwareAccelerated() const override;
+			SIGNALPROCESSING_API virtual bool IsHardwareAccelerated() const override;
 
 			// If true, this implementation requires input and output arrays to be 128 bit aligned.
-			virtual bool Expects128BitAlignedArrays() const override;
+			SIGNALPROCESSING_API virtual bool Expects128BitAlignedArrays() const override;
 
 			// Returns true if the input settings are supported by this factory.
-			virtual bool AreFFTSettingsSupported(const FFFTSettings& InSettings) const override;
+			SIGNALPROCESSING_API virtual bool AreFFTSettingsSupported(const FFFTSettings& InSettings) const override;
 
 			// Create a new FFT algorithm.
-			virtual TUniquePtr<IFFTAlgorithm> NewFFTAlgorithm(const FFFTSettings& InSettings) override;
+			SIGNALPROCESSING_API virtual TUniquePtr<IFFTAlgorithm> NewFFTAlgorithm(const FFFTSettings& InSettings) override;
 	};
 
 	struct FrequencyBuffer
@@ -206,26 +206,26 @@ namespace Audio
 	SIGNALPROCESSING_API void CrossCorrelate(FrequencyBuffer& FirstBufferFrequencies, FrequencyBuffer& SecondBufferFrequencies, int32 NumSamples, FrequencyBuffer& OutCorrelation);
 
 	// Deprecated in 5.1
-	class SIGNALPROCESSING_API FFFTConvolver_DEPRECATED
+	class FFFTConvolver_DEPRECATED
 	{
 	public:
-		FFFTConvolver_DEPRECATED();
+		SIGNALPROCESSING_API FFFTConvolver_DEPRECATED();
 
 		/*
 		 * Applies the convolver's internal window to InputAudio. Until SetWindow is called, ProcessAudio will not affect InputAudio.
 		 * InputAudio must be a power of two.
 		 */
-		void ProcessAudio(float* InputAudio, int32 NumSamples);
+		SIGNALPROCESSING_API void ProcessAudio(float* InputAudio, int32 NumSamples);
 
 		/**
 		 * Resets the filter window. NOT thread safe to call during ProcessAudio.
 		 * This function can be called with a time domain impulse response, or precomputed frequency values. 
 		 * FilterSize must be a power of two.
 		 */ 
-		void SetFilter(const float* InFilterReal, const float* InFilterImag, int32 FilterSize, int32 FFTSize);
-		void SetFilter(const FrequencyBuffer& InFilterFrequencies, int32 FilterSize);
-		void SetFilter(const float* TimeDomainBuffer, int32 FilterSize);
-		void SetFilter(const FAlignedFloatBuffer& TimeDomainBuffer);
+		SIGNALPROCESSING_API void SetFilter(const float* InFilterReal, const float* InFilterImag, int32 FilterSize, int32 FFTSize);
+		SIGNALPROCESSING_API void SetFilter(const FrequencyBuffer& InFilterFrequencies, int32 FilterSize);
+		SIGNALPROCESSING_API void SetFilter(const float* TimeDomainBuffer, int32 FilterSize);
+		SIGNALPROCESSING_API void SetFilter(const FAlignedFloatBuffer& TimeDomainBuffer);
 
 	private:
 		void ConvolveBlock(float* InputAudio, int32 NumSamples);
diff --git a/Engine/Source/Runtime/SignalProcessing/Public/DSP/AudioLinearCurve.h b/Engine/Source/Runtime/SignalProcessing/Public/DSP/AudioLinearCurve.h
--- a/Engine/Source/Runtime/SignalProcessing/Public/DSP/AudioLinearCurve.h
+++ b/Engine/Source/Runtime/SignalProcessing/Public/DSP/AudioLinearCurve.h
@@ -9,41 +9,41 @@ namespace Audio
 {
 	// A simple curve class that takes FVector2D points
 	// Required to decouple from Engine module
-	class SIGNALPROCESSING_API FLinearCurve
+	class FLinearCurve
 	{
 	public:
 		FLinearCurve() = default;
-		FLinearCurve(const TArray<FVector2D>& InCurve);
-		FLinearCurve(TArray<FVector2D>&& InCurve);
+		SIGNALPROCESSING_API FLinearCurve(const TArray<FVector2D>& InCurve);
+		SIGNALPROCESSING_API FLinearCurve(TArray<FVector2D>&& InCurve);
 
 		virtual ~FLinearCurve() = default;
 
 		// Sets curve array
-		void SetCurvePoints(const TArray<FVector2D>& InPoints);
-		void SetCurvePoints(TArray<FVector2D>&& InPoints);
+		SIGNALPROCESSING_API void SetCurvePoints(const TArray<FVector2D>& InPoints);
+		SIGNALPROCESSING_API void SetCurvePoints(TArray<FVector2D>&& InPoints);
 
 		// Clears the points array
-		void ClearPoints();
+		SIGNALPROCESSING_API void ClearPoints();
 
 		// Returns the number of points in the curve
 		int32 Num() const { return Points.Num(); }
 
 		// Adds individual points
-		void AddPoint(const FVector2D& InPoint);
-		void AddPoints(const TArray<FVector2D>& InPoints);
+		SIGNALPROCESSING_API void AddPoint(const FVector2D& InPoint);
+		SIGNALPROCESSING_API void AddPoints(const TArray<FVector2D>& InPoints);
 
 		// Retrieves curve array
-		const TArray<FVector2D>& GetCurvePoints() const;
+		SIGNALPROCESSING_API const TArray<FVector2D>& GetCurvePoints() const;
 
 		// Evaluates the curve with the given domain. 
 		// Returns false if the curve was incapable of evaluation (i.e. empty).
 		// Will clamp the output to the values of the domain
-		bool Eval(float InDomain, float& OutValue) const;
+		SIGNALPROCESSING_API bool Eval(float InDomain, float& OutValue) const;
 
 	private:
 
 		// Will sort the curve according to increasing domain values
-		void Sort();
+		SIGNALPROCESSING_API void Sort();
 		TArray<FVector2D> Points;
 	};
-}
\ No newline at end of file
+}
diff --git a/Engine/Source/Runtime/SignalProcessing/Public/DSP/BitCrusher.h b/Engine/Source/Runtime/SignalProcessing/Public/DSP/BitCrusher.h
--- a/Engine/Source/Runtime/SignalProcessing/Public/DSP/BitCrusher.h
+++ b/Engine/Source/Runtime/SignalProcessing/Public/DSP/BitCrusher.h
@@ -9,30 +9,30 @@ namespace Audio
 {
 	// Bit crushing effect
 	// https://en.wikipedia.org/wiki/Bitcrusher
-	class SIGNALPROCESSING_API FBitCrusher
+	class FBitCrusher
 	{
 	public:
 		// Constructor
-		FBitCrusher();
+		SIGNALPROCESSING_API FBitCrusher();
 
 		// Destructor
-		~FBitCrusher();
+		SIGNALPROCESSING_API ~FBitCrusher();
 
 		// Initialize the equalizer
-		void Init(const float InSampleRate, const int32 InNumChannels);
+		SIGNALPROCESSING_API void Init(const float InSampleRate, const int32 InNumChannels);
 
 		// The amount to reduce the sample rate of the audio stream.
-		void SetSampleRateCrush(const float InFrequency);
+		SIGNALPROCESSING_API void SetSampleRateCrush(const float InFrequency);
 
 		// The amount to reduce the bit depth of the audio stream.
-		void SetBitDepthCrush(const float InBitDepth);
+		SIGNALPROCESSING_API void SetBitDepthCrush(const float InBitDepth);
 
 		// Processes audio
-		void ProcessAudioFrame(const float* InFrame, float* OutFrame);
-		void ProcessAudio(const float* InBuffer, const int32 InNumSamples, float* OutBuffer);
+		SIGNALPROCESSING_API void ProcessAudioFrame(const float* InFrame, float* OutFrame);
+		SIGNALPROCESSING_API void ProcessAudio(const float* InBuffer, const int32 InNumSamples, float* OutBuffer);
 
 		// Returns the maximum value given bit depths will be clamped to.
-		static float GetMaxBitDepth();
+		static SIGNALPROCESSING_API float GetMaxBitDepth();
 
 	private:
 		// i.e. 8 bit, etc. But can be float!
diff --git a/Engine/Source/Runtime/SignalProcessing/Public/DSP/BlockCorrelator.h b/Engine/Source/Runtime/SignalProcessing/Public/DSP/BlockCorrelator.h
--- a/Engine/Source/Runtime/SignalProcessing/Public/DSP/BlockCorrelator.h
+++ b/Engine/Source/Runtime/SignalProcessing/Public/DSP/BlockCorrelator.h
@@ -23,20 +23,20 @@ namespace Audio
 	};
 
 	/** Correlation calculator which utilizes FFT to perform fast correlation calculations. */
-	class SIGNALPROCESSING_API FBlockCorrelator
+	class FBlockCorrelator
 	{
 		public:
 			/** Construct a block correlator with FBlockCorrelatorSettings */
-			FBlockCorrelator(const FBlockCorrelatorSettings& InSettings);
+			SIGNALPROCESSING_API FBlockCorrelator(const FBlockCorrelatorSettings& InSettings);
 
 			/** Return settings used to construct block correlator. */
-			const FBlockCorrelatorSettings& GetSettings() const;
+			SIGNALPROCESSING_API const FBlockCorrelatorSettings& GetSettings() const;
 
 			/** Returns the number of input values to use when calling CrossCorrelate(...) or AutoCorrelate(...) */
-			int32 GetNumInputValues() const;
+			SIGNALPROCESSING_API int32 GetNumInputValues() const;
 
 			/** Returns the number of output values to use when calling CrossCorrelate(...) or AutoCorrelate(...) */
-			int32 GetNumOutputValues() const;
+			SIGNALPROCESSING_API int32 GetNumOutputValues() const;
 
 			/** Cross correlate two input signals. 
 			 *
@@ -53,7 +53,7 @@ namespace Audio
 			 * @param InputB - Second input block with `GetNumInputValues()` elements. 
 			 * @param Output - Output block with `GetNumOutputValues()` elements.
 			 */
-			void CrossCorrelate(const FAlignedFloatBuffer& InputA, const FAlignedFloatBuffer& InputB, FAlignedFloatBuffer& Output);
+			SIGNALPROCESSING_API void CrossCorrelate(const FAlignedFloatBuffer& InputA, const FAlignedFloatBuffer& InputB, FAlignedFloatBuffer& Output);
 
 			/** Autocorrelate an input signal. 
 			 *
@@ -68,7 +68,7 @@ namespace Audio
 			 * @param Input - input block with `GetNumInputValues()` elements. 
 			 * @param Output - Output block with `GetNumOutputValues()` elements.
 			 */
-			void AutoCorrelate(const FAlignedFloatBuffer& Input, FAlignedFloatBuffer& Output);
+			SIGNALPROCESSING_API void AutoCorrelate(const FAlignedFloatBuffer& Input, FAlignedFloatBuffer& Output);
 
 
 		private:
diff --git a/Engine/Source/Runtime/SignalProcessing/Public/DSP/BufferOnePoleLPF.h b/Engine/Source/Runtime/SignalProcessing/Public/DSP/BufferOnePoleLPF.h
--- a/Engine/Source/Runtime/SignalProcessing/Public/DSP/BufferOnePoleLPF.h
+++ b/Engine/Source/Runtime/SignalProcessing/Public/DSP/BufferOnePoleLPF.h
@@ -10,27 +10,27 @@ namespace Audio
 {
 
 	// Simple 1-pole lowpass filter
-	class SIGNALPROCESSING_API FBufferOnePoleLPF
+	class FBufferOnePoleLPF
 	{
 	public:
 
 		// Constructor 
-		FBufferOnePoleLPF(float InG = 0.0f);
+		SIGNALPROCESSING_API FBufferOnePoleLPF(float InG = 0.0f);
 
 		// Set the LPF gain coefficient
-		void SetG(float InG);
+		SIGNALPROCESSING_API void SetG(float InG);
 
 		// Resets the sample delay to 0
-		void Reset();
+		SIGNALPROCESSING_API void Reset();
 
 		// Resets the audio to silence.
-		void FlushAudio();
+		SIGNALPROCESSING_API void FlushAudio();
 
 		// Sets the filter frequency using normalized frequency (between 0.0 and 1.0f or 0.0 hz and Nyquist Frequency in Hz) 
-		void SetFrequency(const float InFrequency);
+		SIGNALPROCESSING_API void SetFrequency(const float InFrequency);
 
 		// Process InSamples and place filtered audio in OutSamples
-		void ProcessAudio(const FAlignedFloatBuffer& InSamples, FAlignedFloatBuffer& OutSamples);
+		SIGNALPROCESSING_API void ProcessAudio(const FAlignedFloatBuffer& InSamples, FAlignedFloatBuffer& OutSamples);
 
 	protected:
 		float CutoffFrequency;
diff --git a/Engine/Source/Runtime/SignalProcessing/Public/DSP/BufferVectorOperations.h b/Engine/Source/Runtime/SignalProcessing/Public/DSP/BufferVectorOperations.h
--- a/Engine/Source/Runtime/SignalProcessing/Public/DSP/BufferVectorOperations.h
+++ b/Engine/Source/Runtime/SignalProcessing/Public/DSP/BufferVectorOperations.h
@@ -303,24 +303,24 @@ namespace Audio
 	SIGNALPROCESSING_API void BufferSum2ChannelToMonoFast(const float* RESTRICT InSamples, float* RESTRICT OutSamples, const int32 InNumFrames);
 
 	/** Class which handles a vectorized interpolation of an entire buffer to the values of a target buffer */
-	class SIGNALPROCESSING_API FBufferLinearEase
+	class FBufferLinearEase
 	{
 	public:
-		FBufferLinearEase();
-		FBufferLinearEase(const FAlignedFloatBuffer& InSourceValues, const FAlignedFloatBuffer& InTargetValues, int32 InLerpLength);
-		~FBufferLinearEase();
+		SIGNALPROCESSING_API FBufferLinearEase();
+		SIGNALPROCESSING_API FBufferLinearEase(const FAlignedFloatBuffer& InSourceValues, const FAlignedFloatBuffer& InTargetValues, int32 InLerpLength);
+		SIGNALPROCESSING_API ~FBufferLinearEase();
 
 		/** will cache SourceValues ptr and manually update SourceValues on Update() */
-		void Init(const FAlignedFloatBuffer& InSourceValues, const FAlignedFloatBuffer& InTargetValues, int32 InLerpLength);
+		SIGNALPROCESSING_API void Init(const FAlignedFloatBuffer& InSourceValues, const FAlignedFloatBuffer& InTargetValues, int32 InLerpLength);
 
 		/** Performs Vectorized update of SourceValues float buffer. Returns true if interpolation is complete */
-		bool Update(FAlignedFloatBuffer& InSourceValues);
+		SIGNALPROCESSING_API bool Update(FAlignedFloatBuffer& InSourceValues);
 
 		/** Update overloaded to let you jump forward more than a single time-step */
-		bool Update(uint32 StepsToJumpForward, FAlignedFloatBuffer& InSourceValues);
+		SIGNALPROCESSING_API bool Update(uint32 StepsToJumpForward, FAlignedFloatBuffer& InSourceValues);
 
 		/** returns const reference to the deltas buffer for doing interpolation elsewhere */
-		const FAlignedFloatBuffer& GetDeltaBuffer();
+		SIGNALPROCESSING_API const FAlignedFloatBuffer& GetDeltaBuffer();
 
 	private:
 		int32 BufferLength {0};
diff --git a/Engine/Source/Runtime/SignalProcessing/Public/DSP/Chorus.h b/Engine/Source/Runtime/SignalProcessing/Public/DSP/Chorus.h
--- a/Engine/Source/Runtime/SignalProcessing/Public/DSP/Chorus.h
+++ b/Engine/Source/Runtime/SignalProcessing/Public/DSP/Chorus.h
@@ -19,27 +19,27 @@ namespace Audio
 		};
 	}
 
-	class SIGNALPROCESSING_API FChorus
+	class FChorus
 	{
 	public:
-		FChorus();
-		~FChorus();
+		SIGNALPROCESSING_API FChorus();
+		SIGNALPROCESSING_API ~FChorus();
 
-		void Init(const float InSampleRate, const int32 InNumChannels, const float InDelayLengthSec = 2.0f, const int32 InControlSamplePeriod = 256);
+		SIGNALPROCESSING_API void Init(const float InSampleRate, const int32 InNumChannels, const float InDelayLengthSec = 2.0f, const int32 InControlSamplePeriod = 256);
 
 		FORCEINLINE int32 GetNumChannels() const
 		{
 			return NumChannels;
 		}
 
-		void SetDepth(const EChorusDelays::Type InType, const float InDepth);
-		void SetFrequency(const EChorusDelays::Type InType, const float InFrequency);
-		void SetFeedback(const EChorusDelays::Type InType, const float InFeedback);
-		void SetWetLevel(const float InWetLevel);
-		void SetDryLevel(const float InDryLevel);
-		void SetSpread(const float InSpread);
-		void ProcessAudioFrame(const float* InFrame, float* OutFrame);
-		void ProcessAudio(const float* InBuffer, const int32 InNumSamples, float* OutBuffer);
+		SIGNALPROCESSING_API void SetDepth(const EChorusDelays::Type InType, const float InDepth);
+		SIGNALPROCESSING_API void SetFrequency(const EChorusDelays::Type InType, const float InFrequency);
+		SIGNALPROCESSING_API void SetFeedback(const EChorusDelays::Type InType, const float InFeedback);
+		SIGNALPROCESSING_API void SetWetLevel(const float InWetLevel);
+		SIGNALPROCESSING_API void SetDryLevel(const float InDryLevel);
+		SIGNALPROCESSING_API void SetSpread(const float InSpread);
+		SIGNALPROCESSING_API void ProcessAudioFrame(const float* InFrame, float* OutFrame);
+		SIGNALPROCESSING_API void ProcessAudio(const float* InBuffer, const int32 InNumSamples, float* OutBuffer);
 
 	protected:
 		FDelay Delays[EChorusDelays::NumDelayTypes];
diff --git a/Engine/Source/Runtime/SignalProcessing/Public/DSP/ConstantQ.h b/Engine/Source/Runtime/SignalProcessing/Public/DSP/ConstantQ.h
--- a/Engine/Source/Runtime/SignalProcessing/Public/DSP/ConstantQ.h
+++ b/Engine/Source/Runtime/SignalProcessing/Public/DSP/ConstantQ.h
@@ -20,7 +20,7 @@ namespace Audio
 	};
 
 	/** Settings for Pseudo Constant Q Kernel generation. */
-	struct SIGNALPROCESSING_API FPseudoConstantQKernelSettings
+	struct FPseudoConstantQKernelSettings
 	{
 		int32 NumBands;									//< Total number of resulting constant Q bands.
 		float NumBandsPerOctave;							//< Number of bands to space within an octave.
@@ -57,7 +57,7 @@ namespace Audio
 	SIGNALPROCESSING_API TUniquePtr<FContiguousSparse2DKernelTransform> NewPseudoConstantQKernelTransform(const FPseudoConstantQKernelSettings& InSettings, const int32 InFFTSize, const float InSampleRate);
 
 	/** Settings for a single constant q band. */
-	struct SIGNALPROCESSING_API FPseudoConstantQBandSettings
+	struct FPseudoConstantQBandSettings
 	{
 		/** Center frequency of band in Hz. */
 		float CenterFreq;
@@ -75,15 +75,15 @@ namespace Audio
 		EPseudoConstantQNormalization Normalization;
 	};
 
-	struct SIGNALPROCESSING_API FPseudoConstantQ
+	struct FPseudoConstantQ
 	{
 		// Returns the center frequency of a given CQT band.
-		static float GetConstantQCenterFrequency(const int32 InBandIndex, const float InBaseFrequency, const float InBandsPerOctave);
+		static SIGNALPROCESSING_API float GetConstantQCenterFrequency(const int32 InBandIndex, const float InBaseFrequency, const float InBandsPerOctave);
 
 		// Returns the bandwidth for a given CQT band.
-		static float GetConstantQBandWidth(const float InBandCenter, const float InBandsPerOctave, const float InBandWidthStretch);
+		static SIGNALPROCESSING_API float GetConstantQBandWidth(const float InBandCenter, const float InBandsPerOctave, const float InBandWidthStretch);
 
 		// Fills an array with the spectrum weights to apply to an FFT magnitude or power spectrum output.
-		static void FillArrayWithConstantQBand(const FPseudoConstantQBandSettings& InSettings, FAlignedFloatBuffer& OutOffsetArray, int32& OutOffsetIndex);
+		static SIGNALPROCESSING_API void FillArrayWithConstantQBand(const FPseudoConstantQBandSettings& InSettings, FAlignedFloatBuffer& OutOffsetArray, int32& OutOffsetIndex);
 	};
 }
diff --git a/Engine/Source/Runtime/SignalProcessing/Public/DSP/ConvolutionAlgorithm.h b/Engine/Source/Runtime/SignalProcessing/Public/DSP/ConvolutionAlgorithm.h
--- a/Engine/Source/Runtime/SignalProcessing/Public/DSP/ConvolutionAlgorithm.h
+++ b/Engine/Source/Runtime/SignalProcessing/Public/DSP/ConvolutionAlgorithm.h
@@ -52,7 +52,7 @@ namespace Audio
 	 *
 	 * Interface for Convolution algorithm.
 	 */
-	class SIGNALPROCESSING_API IConvolutionAlgorithm 
+	class IConvolutionAlgorithm 
 	{
 		public:
 
@@ -111,13 +111,13 @@ namespace Audio
 	 *
 	 * Factory interface for creating convolution algorithms.
 	 */
-	class SIGNALPROCESSING_API IConvolutionAlgorithmFactory : public IModularFeature
+	class IConvolutionAlgorithmFactory : public IModularFeature
 	{
 		public:
-			virtual ~IConvolutionAlgorithmFactory();
+			SIGNALPROCESSING_API virtual ~IConvolutionAlgorithmFactory();
 
 			/** Name of modular feature for Convolution factory.  */
-			static const FName GetModularFeatureName();
+			static SIGNALPROCESSING_API const FName GetModularFeatureName();
 
 			/** Name of this particular factory. */
 			virtual const FName GetFactoryName() const = 0;
@@ -136,11 +136,11 @@ namespace Audio
 	 *
 	 * FConvolutionFactory creates convolution algorithms. It will choose hardware accelerated versions when they are available. 
 	 */
-	class SIGNALPROCESSING_API FConvolutionFactory
+	class FConvolutionFactory
 	{
 		public:
 			/** This denotes that no specific IConvolutionAlgorithmFactory is desired. */
-			static const FName AnyAlgorithmFactory;
+			static SIGNALPROCESSING_API const FName AnyAlgorithmFactory;
 
 			/** NewConvolutionAlgorithm
 			 *
@@ -150,6 +150,6 @@ namespace Audio
 			 * @param InAlgorithmFactoryName - If not equal to FConvolutionFactory::AnyAlgorithmFactory, will only uses Convolution algorithm facotry if IConvolutionAlgorithmFactory::GetFactoryName() equals InAlgorithmFactoryName.
 			 * @return A TUniquePtr<IConvolutionAlgorithm> to the created Convolution. This pointer can be invalid if an error occured or the fft algorithm could not be created.
 			 */
-			static TUniquePtr<IConvolutionAlgorithm> NewConvolutionAlgorithm(const FConvolutionSettings& InSettings, const FName& InAlgorithmFactoryName = AnyAlgorithmFactory);
+			static SIGNALPROCESSING_API TUniquePtr<IConvolutionAlgorithm> NewConvolutionAlgorithm(const FConvolutionSettings& InSettings, const FName& InAlgorithmFactoryName = AnyAlgorithmFactory);
 	};
 }
diff --git a/Engine/Source/Runtime/SignalProcessing/Public/DSP/Delay.h b/Engine/Source/Runtime/SignalProcessing/Public/DSP/Delay.h
--- a/Engine/Source/Runtime/SignalProcessing/Public/DSP/Delay.h
+++ b/Engine/Source/Runtime/SignalProcessing/Public/DSP/Delay.h
@@ -9,56 +9,56 @@
 namespace Audio
 {
 	// Circular Buffer Delay Line
-	class SIGNALPROCESSING_API FDelay
+	class FDelay
 	{
 	public:
 		// Constructor
-		FDelay();
+		SIGNALPROCESSING_API FDelay();
 
 		// Virtual Destructor
-		virtual ~FDelay();
+		SIGNALPROCESSING_API virtual ~FDelay();
 
 	public:
 		// Initialization of the delay with given sample rate and max buffer size in samples.
-		void Init(const float InSampleRate, const float InBufferLengthSec = 2.0f);
+		SIGNALPROCESSING_API void Init(const float InSampleRate, const float InBufferLengthSec = 2.0f);
 
 		// Resets the delay line state, flushes buffer and resets read/write pointers.
-		void Reset();
+		SIGNALPROCESSING_API void Reset();
 
 		// Sets the delay line length. Will clamp to within range of the max initialized delay line length (won't resize).
-		void SetDelayMsec(const float InDelayMsec);
+		SIGNALPROCESSING_API void SetDelayMsec(const float InDelayMsec);
 
 		// Same as SetDelayMsec, except in samples.
-		void SetDelaySamples(const float InDelaySamples);
+		SIGNALPROCESSING_API void SetDelaySamples(const float InDelaySamples);
 
 		// Sets the delay line length but using the internal easing function for smooth delay line interpolation.
-		void SetEasedDelayMsec(const float InDelayMsec, const bool bIsInit = false);
+		SIGNALPROCESSING_API void SetEasedDelayMsec(const float InDelayMsec, const bool bIsInit = false);
 
 		// Sets the easing factor for the delay line's internal exponential interpolator.
-		void SetEaseFactor(const float InEaseFactor);
+		SIGNALPROCESSING_API void SetEaseFactor(const float InEaseFactor);
 
 		// Sets the output attenuation in DB
-		void SetOutputAttenuationDB(const float InDelayAttenDB);
+		SIGNALPROCESSING_API void SetOutputAttenuationDB(const float InDelayAttenDB);
 
 		// Returns the current delay line length (in samples).
 		float GetDelayLengthSamples() const { return DelayInSamples; }
 
 		// Reads the delay line at current read index without writing or incrementing read/write pointers.
-		float Read() const;
+		SIGNALPROCESSING_API float Read() const;
 
 		// Reads the delay line at an arbitrary time in Msec without writing or incrementing read/write pointers.
-		float ReadDelayAt(const float InReadMsec) const;
+		SIGNALPROCESSING_API float ReadDelayAt(const float InReadMsec) const;
 
 		// Write the input and increment read/write pointers
-		void WriteDelayAndInc(const float InDelayInput);
+		SIGNALPROCESSING_API void WriteDelayAndInc(const float InDelayInput);
 
 		// Process audio in the delay line, return the delayed value
-		virtual float ProcessAudioSample(const float InAudio);
+		SIGNALPROCESSING_API virtual float ProcessAudioSample(const float InAudio);
 
 	protected:
 
 		// Updates delay line based on any recent changes to settings
-		void Update(bool bForce = false);
+		SIGNALPROCESSING_API void Update(bool bForce = false);
 
 		// Pointer to the circular buffer of audio.
 		float* AudioBuffer;
diff --git a/Engine/Source/Runtime/SignalProcessing/Public/DSP/DelayStereo.h b/Engine/Source/Runtime/SignalProcessing/Public/DSP/DelayStereo.h
--- a/Engine/Source/Runtime/SignalProcessing/Public/DSP/DelayStereo.h
+++ b/Engine/Source/Runtime/SignalProcessing/Public/DSP/DelayStereo.h
@@ -21,54 +21,54 @@ namespace Audio
 		};
 	}
 
-	class SIGNALPROCESSING_API FDelayStereo
+	class FDelayStereo
 	{
 	public:
-		FDelayStereo();
-		~FDelayStereo();
+		SIGNALPROCESSING_API FDelayStereo();
+		SIGNALPROCESSING_API ~FDelayStereo();
 
 		// Initializes the stereo delay with given sample rate and default max delay length
-		void Init(const float InSampleRate, const int32 InNumChannels, const float InDelayLengthSec = 2.0f);
+		SIGNALPROCESSING_API void Init(const float InSampleRate, const int32 InNumChannels, const float InDelayLengthSec = 2.0f);
 
 		// Resets the stereo delay state
-		void Reset();
+		SIGNALPROCESSING_API void Reset();
 
 		// Process a single frame of audio
-		void ProcessAudioFrame(const float* InFrame, float* OutFrame);
+		SIGNALPROCESSING_API void ProcessAudioFrame(const float* InFrame, float* OutFrame);
 
 		// Process a buffer of audio
-		void ProcessAudio(const float* InBuffer, const int32 InNumSamples, float* OutBuffer);
+		SIGNALPROCESSING_API void ProcessAudio(const float* InBuffer, const int32 InNumSamples, float* OutBuffer);
 
 		// Sets which delay stereo mode to use
-		void SetMode(const EStereoDelayMode::Type InMode);
+		SIGNALPROCESSING_API void SetMode(const EStereoDelayMode::Type InMode);
 
 		// Gets the current stereo dealy mode
 		EStereoDelayMode::Type GetMode() const { return DelayMode; }
 
 		// Sets the delay time in msec. 
-		void SetDelayTimeMsec(const float InDelayTimeMsec);
+		SIGNALPROCESSING_API void SetDelayTimeMsec(const float InDelayTimeMsec);
 
 		// Sets the feedback amount
-		void SetFeedback(const float InFeedback);
+		SIGNALPROCESSING_API void SetFeedback(const float InFeedback);
 
 		// Sets the delay ratio (scales difference between left and right stereo delays)
-		void SetDelayRatio(const float InDelayRatio);
+		SIGNALPROCESSING_API void SetDelayRatio(const float InDelayRatio);
 
 		// Sets the amount of the effect to mix in the output
-		void SetWetLevel(const float InWetLevel);
+		SIGNALPROCESSING_API void SetWetLevel(const float InWetLevel);
 
 		// Sets the amount of the effect to mix in the output
-		void SetDryLevel(const float InDryLevel);
+		SIGNALPROCESSING_API void SetDryLevel(const float InDryLevel);
 
 		// Sets whether or not the filter is enabled
-		void SetFilterEnabled(bool bInEnabled);
+		SIGNALPROCESSING_API void SetFilterEnabled(bool bInEnabled);
 
 		// Sets the filter settings
-		void SetFilterSettings(EBiquadFilter::Type InFilterType, const float InCutoffFrequency, const float InQ);
+		SIGNALPROCESSING_API void SetFilterSettings(EBiquadFilter::Type InFilterType, const float InCutoffFrequency, const float InQ);
 
 	protected:
 		// Updates the delays based on recent parameters
-		void UpdateDelays();
+		SIGNALPROCESSING_API void UpdateDelays();
 
 		// Delay lines per channel
 		TArray<FDelay> Delays;
diff --git a/Engine/Source/Runtime/SignalProcessing/Public/DSP/DirectedGraphAlgo.h b/Engine/Source/Runtime/SignalProcessing/Public/DSP/DirectedGraphAlgo.h
--- a/Engine/Source/Runtime/SignalProcessing/Public/DSP/DirectedGraphAlgo.h
+++ b/Engine/Source/Runtime/SignalProcessing/Public/DSP/DirectedGraphAlgo.h
@@ -34,14 +34,14 @@ namespace Audio
 	typedef TMap<int32, FDirectedTreeElement> FDirectedTree;
 
 
-	struct SIGNALPROCESSING_API FDirectedGraphAlgo
+	struct FDirectedGraphAlgo
 	{
 		/** Build a directed tree from an array of edges.
 		 *
 		 * @parma InEdges - An array of directed eges.
 		 * @param OutTree - A tree structure built from the edges.
 		 */
-		static void BuildDirectedTree(TArrayView<const FDirectedEdge> InEdges, FDirectedTree& OutTree);
+		static SIGNALPROCESSING_API void BuildDirectedTree(TArrayView<const FDirectedEdge> InEdges, FDirectedTree& OutTree);
 
 		/** Build a transpose directed tree from an array of edges. 
 		 *
@@ -50,7 +50,7 @@ namespace Audio
 		 * @parma InEdges - An array of directed eges.
 		 * @param OutTree - A tree structure built from the reversed edges.
 		 */
-		static void BuildTransposeDirectedTree(TArrayView<const FDirectedEdge> InEdges, FDirectedTree& OutTree);
+		static SIGNALPROCESSING_API void BuildTransposeDirectedTree(TArrayView<const FDirectedEdge> InEdges, FDirectedTree& OutTree);
 
 		/** Traverse a tree in a depth first ordering.
 		 *
@@ -62,7 +62,7 @@ namespace Audio
 		 *                          manner. If this function returns false, the 
 		 *                          children of the current vertex are not visited.
 		 */
-		static void DepthFirstTraversal(int32 InInitialVertex, const FDirectedTree& InTree, TFunctionRef<bool (int32)> InVisitFunc);
+		static SIGNALPROCESSING_API void DepthFirstTraversal(int32 InInitialVertex, const FDirectedTree& InTree, TFunctionRef<bool (int32)> InVisitFunc);
 
 		/** Traverse a tree in a breadth first ordering.
 		 *
@@ -75,7 +75,7 @@ namespace Audio
 		 *                          then the children of the current vertex will 
 		 *                          not be visited.
 		 */
-		static void BreadthFirstTraversal(int32 InInitialVertex, const FDirectedTree& InTree, TFunctionRef<bool (int32)> InVisitFunc);
+		static SIGNALPROCESSING_API void BreadthFirstTraversal(int32 InInitialVertex, const FDirectedTree& InTree, TFunctionRef<bool (int32)> InVisitFunc);
 
 		/** Sort vertices topologically using a depth first sorting algorithm.
 		 *
@@ -85,7 +85,7 @@ namespace Audio
 		 *
 		 * @return True if sorting was successful. False otherwise.
 		 */
-		static bool DepthFirstTopologicalSort(TArrayView<const int32> InUniqueVertices, TArrayView<const FDirectedEdge> InUniqueEdges, TArray<int32>& OutVertexOrder);
+		static SIGNALPROCESSING_API bool DepthFirstTopologicalSort(TArrayView<const int32> InUniqueVertices, TArrayView<const FDirectedEdge> InUniqueEdges, TArray<int32>& OutVertexOrder);
 
 		/** Sort vertices topologically using a Kahn's sorting algorithm.
 		 *
@@ -95,7 +95,7 @@ namespace Audio
 		 *
 		 * @return True if sorting was successful. False otherwise.
 		 */
-		static bool KahnTopologicalSort(TArrayView<const int32> InUniqueVertices, TArrayView<const FDirectedEdge> InUniqueEdges, TArray<int32>& OutVertexOrder);
+		static SIGNALPROCESSING_API bool KahnTopologicalSort(TArrayView<const int32> InUniqueVertices, TArrayView<const FDirectedEdge> InUniqueEdges, TArray<int32>& OutVertexOrder);
 
 
 		/** Find strongly connected components given a set of edges using Tarjan
@@ -107,6 +107,6 @@ namespace Audio
 		 *
 		 * @return True if one or more strongly connected components are added to OutComponents. False otherwise. 
 		 */
-		static bool TarjanStronglyConnectedComponents(const TSet<FDirectedEdge>& InEdges, TArray<FStronglyConnectedComponent>& OutComponents, bool bExcludeSingleVertex = true);
+		static SIGNALPROCESSING_API bool TarjanStronglyConnectedComponents(const TSet<FDirectedEdge>& InEdges, TArray<FStronglyConnectedComponent>& OutComponents, bool bExcludeSingleVertex = true);
 	};
 }
diff --git a/Engine/Source/Runtime/SignalProcessing/Public/DSP/DynamicDelayAPF.h b/Engine/Source/Runtime/SignalProcessing/Public/DSP/DynamicDelayAPF.h
--- a/Engine/Source/Runtime/SignalProcessing/Public/DSP/DynamicDelayAPF.h
+++ b/Engine/Source/Runtime/SignalProcessing/Public/DSP/DynamicDelayAPF.h
@@ -12,7 +12,7 @@ namespace Audio
 {
 	// All Pass Filter with a long fractional delay which can be set per a sample. This filter is specifically designed for 
 	// reverb applications where filter delay lines are long.
-	class SIGNALPROCESSING_API FDynamicDelayAPF
+	class FDynamicDelayAPF
 	{
 	public:
 		// InG is the filter coefficient used in the long delay all pass filter.
@@ -20,10 +20,10 @@ namespace Audio
 		// InMaxDelay is the maximum allowable delay of the all pass filter.
 		// InMaxNumInternalBufferSamples is the maximum internal block size used internally. 
 		// InSampleRate is the current rendering sample rate. Used to convert parameter ease time from seconds to samples.
-		FDynamicDelayAPF(float InG, int32 InMinDelay, int32 InMaxDelay, int32 InMaxNumInternalBufferSamples, float InSampleRate);
+		SIGNALPROCESSING_API FDynamicDelayAPF(float InG, int32 InMinDelay, int32 InMaxDelay, int32 InMaxNumInternalBufferSamples, float InSampleRate);
 
 		// Destructor
-		~FDynamicDelayAPF();
+		SIGNALPROCESSING_API ~FDynamicDelayAPF();
 
 		// Set the APF feedback/feedforward gain coefficient
 		void SetG(float InG)
@@ -33,17 +33,17 @@ namespace Audio
 
 		// Processes InSamples through the all pass filter and populates OutSamples with the filter output.
 		// InDelays denotes the per-sample delay of the allpass. It must have an equal number of elements as InSamples
-		void ProcessAudio(const FAlignedFloatBuffer& InSamples, const FAlignedFloatBuffer& InDelays, FAlignedFloatBuffer& OutSamples);
+		SIGNALPROCESSING_API void ProcessAudio(const FAlignedFloatBuffer& InSamples, const FAlignedFloatBuffer& InDelays, FAlignedFloatBuffer& OutSamples);
 
 		// Zeros the internal delay line.
-		void Reset();
+		SIGNALPROCESSING_API void Reset();
 
 		void SetEaseTimeInSec(float InEaseTimeInSec) { EaseTimeInSec = InEaseTimeInSec; }
 
 	protected:
 
 		// Process one block of audio.
-		void ProcessAudioBlock(const float* InSamples, const FAlignedFloatBuffer& InFractionalDelays, const int32 InNum, float* OutSamples);
+		SIGNALPROCESSING_API void ProcessAudioBlock(const float* InSamples, const FAlignedFloatBuffer& InFractionalDelays, const int32 InNum, float* OutSamples);
 
 	private:
 		// Feedback/Feedforward gain coefficient
diff --git a/Engine/Source/Runtime/SignalProcessing/Public/DSP/DynamicStateVariableFilter.h b/Engine/Source/Runtime/SignalProcessing/Public/DSP/DynamicStateVariableFilter.h
--- a/Engine/Source/Runtime/SignalProcessing/Public/DSP/DynamicStateVariableFilter.h
+++ b/Engine/Source/Runtime/SignalProcessing/Public/DSP/DynamicStateVariableFilter.h
@@ -15,30 +15,30 @@ namespace Audio
 	};
 
 	// combination of two state-variable filters, the first filtering a key signal, which drives gain in the second
-	class SIGNALPROCESSING_API FDynamicStateVariableFilter
+	class FDynamicStateVariableFilter
 	{
 	public:
-		void Init(const float InSampleRate, const int32 InNumChannels);
-		void Reset();
+		SIGNALPROCESSING_API void Init(const float InSampleRate, const int32 InNumChannels);
+		SIGNALPROCESSING_API void Reset();
 
-		void ProcessAudio(const float* InSamples, float* OutSamples, const int32 InNumSamples);
-		void ProcessAudio(const float* InSamples, float* OutSamples, const float* KeySamples, const int32 InNumSamples);
+		SIGNALPROCESSING_API void ProcessAudio(const float* InSamples, float* OutSamples, const int32 InNumSamples);
+		SIGNALPROCESSING_API void ProcessAudio(const float* InSamples, float* OutSamples, const float* KeySamples, const int32 InNumSamples);
 
-		void SetFrequency(const float InFrequency);
-		void SetQ(const float InQ);
-		void SetGain(const float InGain);
-		void SetFilterType(const EDynamicFilterType InFilterType);
+		SIGNALPROCESSING_API void SetFrequency(const float InFrequency);
+		SIGNALPROCESSING_API void SetQ(const float InQ);
+		SIGNALPROCESSING_API void SetGain(const float InGain);
+		SIGNALPROCESSING_API void SetFilterType(const EDynamicFilterType InFilterType);
 
-		void SetAnalog(const bool bInAnalog);
-		void SetAttackTime(const float InAttackTime);
-		void SetReleaseTime(const float InReleaseTime);
-		void SetThreshold(const float InThresholdDb);
-		void SetEnvMode(const EPeakMode::Type InMode);
-		void SetDynamicRange(const float InDynamicRange);
-		void SetRatio(const float InRatio);
-		void SetKnee(const float InKnee);
+		SIGNALPROCESSING_API void SetAnalog(const bool bInAnalog);
+		SIGNALPROCESSING_API void SetAttackTime(const float InAttackTime);
+		SIGNALPROCESSING_API void SetReleaseTime(const float InReleaseTime);
+		SIGNALPROCESSING_API void SetThreshold(const float InThresholdDb);
+		SIGNALPROCESSING_API void SetEnvMode(const EPeakMode::Type InMode);
+		SIGNALPROCESSING_API void SetDynamicRange(const float InDynamicRange);
+		SIGNALPROCESSING_API void SetRatio(const float InRatio);
+		SIGNALPROCESSING_API void SetKnee(const float InKnee);
 
-		float GetGainReduction();
+		SIGNALPROCESSING_API float GetGainReduction();
 
 	protected:
 		float SampleRate = 48000.f;
@@ -116,16 +116,16 @@ namespace Audio
 		// Points in the knee used for lagrangian interpolation
 		TArray<FVector2D> KneePoints;
 		
-		void ProcessBell(const float* InSamples, float* OutSamples, const float* KeySamples, const int32 InNumSamples);
-		void ProcessLowShelf(const float* InSamples, float* OutSamples, const float* KeySamples, const int32 InNumSamples);
-		void ProcessHighShelf(const float* InSamples, float* OutSamples, const float* KeySamples, const int32 InNumSamples);
+		SIGNALPROCESSING_API void ProcessBell(const float* InSamples, float* OutSamples, const float* KeySamples, const int32 InNumSamples);
+		SIGNALPROCESSING_API void ProcessLowShelf(const float* InSamples, float* OutSamples, const float* KeySamples, const int32 InNumSamples);
+		SIGNALPROCESSING_API void ProcessHighShelf(const float* InSamples, float* OutSamples, const float* KeySamples, const int32 InNumSamples);
 		
-		void UpdateSettings();
-		void UpdatePreFilterVariables();
-		void UpdatePostFilterVariables();
+		SIGNALPROCESSING_API void UpdateSettings();
+		SIGNALPROCESSING_API void UpdatePreFilterVariables();
+		SIGNALPROCESSING_API void UpdatePostFilterVariables();
 		
-		float CalcGain(const float KeySample);
+		SIGNALPROCESSING_API float CalcGain(const float KeySample);
 		
-		void ClearFilterDenormals();
+		SIGNALPROCESSING_API void ClearFilterDenormals();
 	};
 }
diff --git a/Engine/Source/Runtime/SignalProcessing/Public/DSP/DynamicsProcessor.h b/Engine/Source/Runtime/SignalProcessing/Public/DSP/DynamicsProcessor.h
--- a/Engine/Source/Runtime/SignalProcessing/Public/DSP/DynamicsProcessor.h
+++ b/Engine/Source/Runtime/SignalProcessing/Public/DSP/DynamicsProcessor.h
@@ -33,56 +33,56 @@ namespace Audio
 
 	// Dynamic range compressor
 	// https://en.wikipedia.org/wiki/Dynamic_range_compression
-	class SIGNALPROCESSING_API FDynamicsProcessor
+	class FDynamicsProcessor
 	{
 	public:
-		FDynamicsProcessor();
-		~FDynamicsProcessor();
-
-		void Init(const float InSampleRate, const int32 InNumChannels = 2);
-
-		int32 GetNumChannels() const;
-		int32 GetKeyNumChannels() const;
-		float GetMaxLookaheadMsec() const;
-
-		void SetLookaheadMsec(const float InLookAheadMsec);
-		void SetAttackTime(const float InAttackTimeMsec);
-		void SetReleaseTime(const float InReleaseTimeMsec);
-		void SetThreshold(const float InThresholdDb);
-		void SetRatio(const float InCompressionRatio);
-		void SetKneeBandwidth(const float InKneeBandwidthDb);
-		void SetInputGain(const float InInputGainDb);
-		void SetKeyAudition(const bool InAuditionEnabled);
-		void SetKeyGain(const float InKeyGain);
-		void SetKeyHighshelfCutoffFrequency(const float InCutoffFreq);
-		void SetKeyHighshelfEnabled(const bool bInEnabled);
-		void SetKeyHighshelfGain(const float InGainDb);
-		void SetKeyLowshelfCutoffFrequency(const float InCutoffFreq);
-		void SetKeyLowshelfEnabled(const bool bInEnabled);
-		void SetKeyLowshelfGain(const float InGainDb);
-		void SetKeyNumChannels(const int32 InNumChannels);
-		void SetNumChannels(const int32 InNumChannels);
-		void SetOutputGain(const float InOutputGainDb);
-		void SetChannelLinkMode(const EDynamicsProcessorChannelLinkMode InLinkMode);
-		void SetAnalogMode(const bool bInIsAnalogMode);
-		void SetPeakMode(const EPeakMode::Type InEnvelopeFollowerModeType);
-		void SetProcessingMode(const EDynamicsProcessingMode::Type ProcessingMode);
-
-		void ProcessAudioFrame(const float* InFrame, float* OutFrame, const float* InKeyFrame);
-		void ProcessAudioFrame(const float* InFrame, float* OutFrame, const float* InKeyFrame, float* OutGain);
-		void ProcessAudio(const float* InBuffer, const int32 InNumSamples, float* OutBuffer, const float* InKeyBuffer = nullptr);
-		void ProcessAudio(const float* InBuffer, const int32 InNumSamples, float* OutBuffer, const float* InKeyBuffer, float* OutEnvelope);
+		SIGNALPROCESSING_API FDynamicsProcessor();
+		SIGNALPROCESSING_API ~FDynamicsProcessor();
+
+		SIGNALPROCESSING_API void Init(const float InSampleRate, const int32 InNumChannels = 2);
+
+		SIGNALPROCESSING_API int32 GetNumChannels() const;
+		SIGNALPROCESSING_API int32 GetKeyNumChannels() const;
+		SIGNALPROCESSING_API float GetMaxLookaheadMsec() const;
+
+		SIGNALPROCESSING_API void SetLookaheadMsec(const float InLookAheadMsec);
+		SIGNALPROCESSING_API void SetAttackTime(const float InAttackTimeMsec);
+		SIGNALPROCESSING_API void SetReleaseTime(const float InReleaseTimeMsec);
+		SIGNALPROCESSING_API void SetThreshold(const float InThresholdDb);
+		SIGNALPROCESSING_API void SetRatio(const float InCompressionRatio);
+		SIGNALPROCESSING_API void SetKneeBandwidth(const float InKneeBandwidthDb);
+		SIGNALPROCESSING_API void SetInputGain(const float InInputGainDb);
+		SIGNALPROCESSING_API void SetKeyAudition(const bool InAuditionEnabled);
+		SIGNALPROCESSING_API void SetKeyGain(const float InKeyGain);
+		SIGNALPROCESSING_API void SetKeyHighshelfCutoffFrequency(const float InCutoffFreq);
+		SIGNALPROCESSING_API void SetKeyHighshelfEnabled(const bool bInEnabled);
+		SIGNALPROCESSING_API void SetKeyHighshelfGain(const float InGainDb);
+		SIGNALPROCESSING_API void SetKeyLowshelfCutoffFrequency(const float InCutoffFreq);
+		SIGNALPROCESSING_API void SetKeyLowshelfEnabled(const bool bInEnabled);
+		SIGNALPROCESSING_API void SetKeyLowshelfGain(const float InGainDb);
+		SIGNALPROCESSING_API void SetKeyNumChannels(const int32 InNumChannels);
+		SIGNALPROCESSING_API void SetNumChannels(const int32 InNumChannels);
+		SIGNALPROCESSING_API void SetOutputGain(const float InOutputGainDb);
+		SIGNALPROCESSING_API void SetChannelLinkMode(const EDynamicsProcessorChannelLinkMode InLinkMode);
+		SIGNALPROCESSING_API void SetAnalogMode(const bool bInIsAnalogMode);
+		SIGNALPROCESSING_API void SetPeakMode(const EPeakMode::Type InEnvelopeFollowerModeType);
+		SIGNALPROCESSING_API void SetProcessingMode(const EDynamicsProcessingMode::Type ProcessingMode);
+
+		SIGNALPROCESSING_API void ProcessAudioFrame(const float* InFrame, float* OutFrame, const float* InKeyFrame);
+		SIGNALPROCESSING_API void ProcessAudioFrame(const float* InFrame, float* OutFrame, const float* InKeyFrame, float* OutGain);
+		SIGNALPROCESSING_API void ProcessAudio(const float* InBuffer, const int32 InNumSamples, float* OutBuffer, const float* InKeyBuffer = nullptr);
+		SIGNALPROCESSING_API void ProcessAudio(const float* InBuffer, const int32 InNumSamples, float* OutBuffer, const float* InKeyBuffer, float* OutEnvelope);
 		 
 
 	protected:
 
-		float ComputeGain(const float InEnvFollowerDb);
+		SIGNALPROCESSING_API float ComputeGain(const float InEnvFollowerDb);
 
 		// Process key frame, returning true if should continue processing
 		// (Returns false in audition mode and writes straight to output).
-		bool ProcessKeyFrame(const float* InKeyFrame, float* OutFrame, bool bKeyIsInput);
+		SIGNALPROCESSING_API bool ProcessKeyFrame(const float* InKeyFrame, float* OutFrame, bool bKeyIsInput);
 
-		bool IsInProcessingThreshold(const float InEnvFollowerDb) const;
+		SIGNALPROCESSING_API bool IsInProcessingThreshold(const float InEnvFollowerDb) const;
 
 		// (Optional) Low-pass filter for input signal
 		FBiquadFilter InputLowshelfFilter;
